Composite: Split HoldoutComposite into CompositeCore & Composite plugins.

* Required step to continue development without the enabled-by-default dependency constraints.
* Composite plugin (with composite actor, pass classes & camera view projection component) is moved to NotForLicencees for 5.6.
* We don't handle asset redirectors, content using materials/MPC will need to be updated.

#virtualized

diff --git a/Engine/Plugins/Experimental/Compositing/CompositeCore/CompositeCore.uplugin b/Engine/Plugins/Experimental/Compositing/CompositeCore/CompositeCore.uplugin
new file mode 100644
--- /dev/null
+++ b/Engine/Plugins/Experimental/Compositing/CompositeCore/CompositeCore.uplugin
@@ -0,0 +1,25 @@
+{
+	"FileVersion": 1,
+	"Version": 1,
+	"VersionName": "1.0",
+	"FriendlyName": "CompositeCore",
+	"Description": "Extensible core plugin for real-time compositing, with a default (holdout) composite pipeline through post-processing.",
+	"Category": "Compositing",
+	"CreatedBy": "Epic Games, Inc.",
+	"CreatedByURL": "https://epicgames.com",
+	"DocsURL": "",
+	"MarketplaceURL": "",
+	"SupportURL": "",
+	"EnabledByDefault": false,
+	"CanContainContent": false,
+	"IsBetaVersion": false,
+	"IsExperimentalVersion": true,
+	"Installed": false,
+	"Modules": [
+		{
+			"Name": "CompositeCore",
+			"Type": "Runtime",
+			"LoadingPhase": "PostConfigInit"
+		}
+	]
+}
\ No newline at end of file
diff --git a/Engine/Plugins/Experimental/Compositing/CompositeCore/Config/BaseCompositeCore.ini b/Engine/Plugins/Experimental/Compositing/CompositeCore/Config/BaseCompositeCore.ini
new file mode 100644
--- /dev/null
+++ b/Engine/Plugins/Experimental/Compositing/CompositeCore/Config/BaseCompositeCore.ini
@@ -0,0 +1,5 @@
+[CoreRedirects]
+
++ClassRedirects=(OldName="/Script/HoldoutComposite.HoldoutCompositeComponent", NewName="/Script/CompositeCore.HoldoutCompositeComponent")
++ClassRedirects=(OldName="/Script/HoldoutComposite.HoldoutCompositeSubsystem", NewName="/Script/CompositeCore.CompositeCoreSubsystem")
++ClassRedirects=(OldName="/Script/HoldoutComposite.HoldoutCompositeSettings", NewName="/Script/CompositeCore.CompositeCoreSettings")
\ No newline at end of file
diff --git a/Engine/Plugins/Experimental/Compositing/CompositeCore/Shaders/Private/CompositeCoreDilate.usf b/Engine/Plugins/Experimental/Compositing/CompositeCore/Shaders/Private/CompositeCoreDilate.usf
new file mode 100644
--- /dev/null
+++ b/Engine/Plugins/Experimental/Compositing/CompositeCore/Shaders/Private/CompositeCoreDilate.usf
@@ -0,0 +1,114 @@
+// Copyright Epic Games, Inc. All Rights Reserved.
+
+#include "/Engine/Private/Common.ush"
+
+#ifndef TILE_SIZE
+    #error TILE_SIZE must be defined
+#endif
+
+#define KERNEL_COUNT 8
+#define ALPHA_THRESHOLD (1.0f - 1e-7f)
+
+// Number of texels arround the group tile for scene color.
+#define LDS_TILE_BORDER_SIZE (1)
+
+// Width in texels of the depth tile cached into LDS.
+#define LDS_TILE_WIDTH (TILE_SIZE + 2 * LDS_TILE_BORDER_SIZE)
+
+// Total number of texels cached in the scene color tile.
+#define LDS_ARRAY_SIZE (LDS_TILE_WIDTH * LDS_TILE_WIDTH)
+
+#define THREADGROUP_TOTAL (TILE_SIZE * TILE_SIZE)
+
+Texture2D InputTexture;
+RWTexture2D<float4> RWOutputTexture;
+uint2 Dimensions;
+
+static const int2 kOffsets[KERNEL_COUNT] = {
+    int2(-1,  0),
+    int2( 1,  0),
+    int2( 0,  1),
+    int2( 0, -1),
+    int2(-1,  1),
+    int2( 1,  1),
+    int2( 1, -1),
+    int2(-1, -1),
+};
+
+groupshared float4 SharedColor[LDS_ARRAY_SIZE];
+
+// Get the index within the LDS array.
+uint GetTileArrayIndexFromPixelOffset(in uint2 GroupThreadId, in int2 PixelOffset)
+{
+    uint2 TilePos = GroupThreadId + uint2(PixelOffset + LDS_TILE_BORDER_SIZE);
+    return TilePos.x + TilePos.y * (LDS_TILE_BORDER_SIZE * 2 + TILE_SIZE);
+}
+
+float4 SampleCachedColors(in uint2 GroupThreadId, int2 PixelOffset = int2(0,0))
+{
+    return SharedColor[GetTileArrayIndexFromPixelOffset(GroupThreadId, PixelOffset)];
+}
+
+[numthreads(TILE_SIZE, TILE_SIZE, 1)]
+void MainCS(
+    uint2 GroupId : SV_GroupID,
+    uint2 GroupThreadId : SV_GroupThreadID,
+    uint GroupThreadIndex : SV_GroupIndex,
+    uint2 DispatchThreadId : SV_DispatchThreadID)
+{
+    // Init LDS
+    {
+        // Global texel coordinate of the active group (corner) offset by the LDS border
+        uint2 GroupTexelOffset = GroupId * uint2(TILE_SIZE, TILE_SIZE) - LDS_TILE_BORDER_SIZE;
+
+        // To account for the dilation margin, each warp thread should load extra texels.
+        // We proceed linearly through the tile (with borders) and each warp thread loads
+        // multiple colors at THREADGROUP_TOTAL intervals with the necessary boundary conditions.
+        const uint LoadCount = (LDS_ARRAY_SIZE + THREADGROUP_TOTAL - 1) / THREADGROUP_TOTAL; // integer ceiling
+
+        uint LinearGroupThreadId = GroupThreadIndex;
+
+        UNROLL
+        for (uint i = 0; i < LoadCount; i++)
+        {
+            uint2 TexelLocation =  GroupTexelOffset + uint2(
+                LinearGroupThreadId % LDS_TILE_WIDTH,
+                LinearGroupThreadId / LDS_TILE_WIDTH);
+
+            if ((LinearGroupThreadId < LDS_ARRAY_SIZE) ||
+                (i != LoadCount - 1) ||
+                (LDS_ARRAY_SIZE % THREADGROUP_TOTAL) == 0)
+            {
+                SharedColor[LinearGroupThreadId] = InputTexture[TexelLocation];
+            }
+            LinearGroupThreadId += THREADGROUP_TOTAL;
+        }
+    }
+
+    GroupMemoryBarrierWithGroupSync();
+
+    // Discard out-of-bound texels
+    if (any(DispatchThreadId >= Dimensions))
+    {
+        return;
+    }
+
+    float4 TexelColor = SampleCachedColors(GroupThreadId);
+    
+    BRANCH
+    if(TexelColor.a > ALPHA_THRESHOLD) // Alpha is inverted opacity
+    {
+        UNROLL
+        for(int i=0; i<KERNEL_COUNT; i++)
+        {
+            float4 OffsetColor = SampleCachedColors(GroupThreadId, kOffsets[i]);
+
+            if(OffsetColor.a < TexelColor.a)
+            {
+                TexelColor = OffsetColor;
+            }
+        }
+    }
+
+    RWOutputTexture[DispatchThreadId] = TexelColor;
+}
diff --git a/Engine/Plugins/Experimental/Compositing/CompositeCore/Shaders/Private/CompositeCoreMerge.usf b/Engine/Plugins/Experimental/Compositing/CompositeCore/Shaders/Private/CompositeCoreMerge.usf
new file mode 100644
--- /dev/null
+++ b/Engine/Plugins/Experimental/Compositing/CompositeCore/Shaders/Private/CompositeCoreMerge.usf
@@ -0,0 +1,187 @@
+// Copyright Epic Games, Inc. All Rights Reserved.
+
+#include "/Engine/Private/Common.ush"
+#include "/Engine/Private/GammaCorrectionCommon.ush"
+#include "/Engine/Private/ScreenPass.ush"
+#include "/Engine/Private/LensDistortion.ush"
+
+#define ESourceEncoding_None  0
+#define ESourceEncoding_Gamma 1
+#define ESourceEncoding_sRGB  2
+
+#define EDistortionUV_None        0
+#define EDistortionUV_Distorted   1
+#define EDistortionUV_Undistorted 2
+
+#define ECompositeCoreMergeOp_Over 0
+#define ECompositeCoreMergeOp_Add  1
+#define ECompositeCoreMergeOp_Mul  2
+#define ECompositeCoreMergeOp_Div  3
+
+SCREEN_PASS_TEXTURE_VIEWPORT(Input0)
+SCREEN_PASS_TEXTURE_VIEWPORT(Input1)
+SCREEN_PASS_TEXTURE_VIEWPORT(Output)
+
+#define INPUT_PASS_TEXTURE_MEMBER(StructName, MemberType, MemberName) MemberType StructName##_##MemberName;
+
+#define INPUT_PASS_TEXTURE(StructName) \
+	INPUT_PASS_TEXTURE_MEMBER(StructName, Texture2D, Texture) \
+	INPUT_PASS_TEXTURE_MEMBER(StructName, SamplerState, Sampler) \
+    INPUT_PASS_TEXTURE_MEMBER(StructName, uint, bAlphaIsOpacity) \
+    INPUT_PASS_TEXTURE_MEMBER(StructName, uint, SourceEncoding) \
+    INPUT_PASS_TEXTURE_MEMBER(StructName, uint, DistortionUV)
+    
+
+INPUT_PASS_TEXTURE(Tex0);
+INPUT_PASS_TEXTURE(Tex1);
+
+Texture2D<float2> DistortingDisplacementTexture;
+SamplerState DistortingDisplacementSampler;
+Texture2D<float2> UndistortingDisplacementTexture;
+SamplerState UndistortingDisplacementSampler;
+
+float2 DisplayGamma; // DisplayGamma, 1/DisplayGamma
+uint MergeOp;
+uint OutputEncoding;
+uint bFlags; // 0: bOutputSceneColor, 1: bUnpremulInput1
+
+float SafeRcp(float x)
+{
+	return x > 0.0 ? rcp(x) : 0.0;
+}
+
+float4 Unpremultiply(float4 InColor)
+{
+	return float4(InColor.xyz * SafeRcp(InColor.w), InColor.w);
+}
+
+float3 Decode(float3 InColor, uint InEncoding)
+{
+    // TODO: HDR cases will need to be handled explicitely.
+    BRANCH
+    if(InEncoding == ESourceEncoding_Gamma)
+    {
+        return pow(InColor, DisplayGamma.x);
+    }
+    else if(InEncoding == ESourceEncoding_sRGB)
+    {
+        return sRGBToLinear(InColor);
+    }
+    else
+    {
+        return InColor;
+    }
+}
+
+float3 Encode(float3 InColor, uint InEncoding)
+{
+    BRANCH
+    if(InEncoding == ESourceEncoding_Gamma)
+    {
+        return pow(InColor, DisplayGamma.y);
+    }
+    else if(InEncoding == ESourceEncoding_sRGB)
+    {
+        return LinearToSrgb(InColor);
+    }
+    else
+    {
+        return InColor;
+    }
+}
+
+float2 CalcInputUVs(float2 ViewportUV, float2 ViewportSize, float2 ViewportMin, float2 ExtentInverse, uint DistortionUVType)
+{
+#if USE_DISTORTION
+    BRANCH
+    if(DistortionUVType == EDistortionUV_Distorted)
+    {
+        ViewportUV = ApplyLensDistortionOnViewportUV(DistortingDisplacementTexture, DistortingDisplacementSampler, ViewportUV);
+    }
+    else if(DistortionUVType == EDistortionUV_Undistorted)
+    {
+        ViewportUV = ApplyLensDistortionOnViewportUV(UndistortingDisplacementTexture, UndistortingDisplacementSampler, ViewportUV);
+    }
+#endif
+
+    return (ViewportUV * ViewportSize + ViewportMin) * ExtentInverse;
+}
+
+float GetOpacity(float InAlpha, uint bAlphaIsOpacity)
+{
+	return select(bAlphaIsOpacity > 0, InAlpha, 1.0 - InAlpha);
+}
+
+float SafeDiv(float A, float B)
+{
+    return B == 0.0 ? 0.0 : A / B;
+}
+
+float4 SafeDiv(float4 A, float4 B)
+{
+    return float4(
+        SafeDiv(A[0], B[0]),
+        SafeDiv(A[1], B[1]),
+        SafeDiv(A[2], B[2]),
+        SafeDiv(A[3], B[3])
+    );
+}
+
+void MainPS(
+    float4 SvPosition : SV_POSITION,
+    out float4 OutColor : SV_Target0
+    )
+{
+    float2 ViewportUV = (SvPosition.xy - Output_ViewportMin) * Output_ViewportSizeInverse;
+
+    // Calculate texture UVs, with/without distortion
+    float2 Tex0_UV = CalcInputUVs(ViewportUV, Input0_ViewportSize, Input0_ViewportMin, Input0_ExtentInverse, Tex0_DistortionUV);
+    float2 Tex1_UV = CalcInputUVs(ViewportUV, Input1_ViewportSize, Input1_ViewportMin, Input1_ExtentInverse, Tex1_DistortionUV);
+
+    // Sample textures
+    float4 Color0 = Tex0_Texture.SampleLevel(Tex0_Sampler, Tex0_UV, 0);
+    float4 Color1 = Tex1_Texture.SampleLevel(Tex1_Sampler, Tex1_UV, 0);
+
+    // Linearize source colors
+    Color0.rgb = Decode(Color0.rgb, Tex0_SourceEncoding);
+    Color1.rgb = Decode(Color1.rgb, Tex1_SourceEncoding);
+
+    // Normalize alpha as opacity
+    Color0.a = GetOpacity(Color0.a, Tex0_bAlphaIsOpacity);
+    Color1.a = GetOpacity(Color1.a, Tex1_bAlphaIsOpacity);
+
+    {
+        // NOTE: This is necessary to extract the solid colors behind translucent alpha holdout masks.
+        // Could be moved elsewhere as a separate pre-processing step, with minor added cost.
+        const bool bUnpremulInput1 = ((bFlags >> 1) & 1);
+        Color1 = select(bUnpremulInput1, Unpremultiply(Color1), Color1);
+    }
+    
+    if(MergeOp == ECompositeCoreMergeOp_Over)
+    {
+        OutColor = Color0 + Color1 * (1.0 - Color0.a);
+    }
+    else if(MergeOp == ECompositeCoreMergeOp_Add)
+    {
+        OutColor = Color0 + Color1;
+    }
+    else if(MergeOp == ECompositeCoreMergeOp_Mul)
+    {
+        OutColor = Color0 * Color1;
+    }
+    else if(MergeOp == ECompositeCoreMergeOp_Div)
+    {
+        OutColor = SafeDiv(Color0, Color1);
+    }
+    else
+    {
+        OutColor = Color0;
+    }
+
+    // Re-encode color
+    OutColor.rgb = Encode(OutColor.rgb, OutputEncoding);
+
+    // Optionally re-invert alpha
+    const bool bOutputSceneColor = (bFlags & 1);
+    OutColor.a = saturate(select(bOutputSceneColor, 1.0 - OutColor.a, OutColor.a));
+}
diff --git a/Engine/Plugins/Experimental/Compositing/CompositeCore/Source/CompositeCore/CompositeCore.Build.cs b/Engine/Plugins/Experimental/Compositing/CompositeCore/Source/CompositeCore/CompositeCore.Build.cs
new file mode 100644
--- /dev/null
+++ b/Engine/Plugins/Experimental/Compositing/CompositeCore/Source/CompositeCore/CompositeCore.Build.cs
@@ -0,0 +1,47 @@
+// Copyright Epic Games, Inc. All Rights Reserved.
+
+using UnrealBuildTool;
+
+public class CompositeCore : ModuleRules
+{
+	public CompositeCore(ReadOnlyTargetRules Target) : base(Target)
+	{
+		// The CompositeCore plugin is distributed with engine hot fixes and thus isn't tied to binary
+		// compatibility between hotfixes by only using Public/ interface of the renderer, but also Internal/ ones.
+		bTreatAsEngineModule = true;
+		PCHUsage = PCHUsageMode.UseExplicitOrSharedPCHs;
+
+		PublicDependencyModuleNames.AddRange(
+			new string[]
+			{
+				"Core",
+				"CoreUObject",
+				"Engine",
+				"DeveloperSettings",
+			}
+		);
+
+		PrivateDependencyModuleNames.AddRange(
+			new string[]
+			{
+				"Projects",
+				"RenderCore",
+				"Renderer",
+				"RHI",
+			}
+		);
+
+		if (Target.bBuildEditor)
+		{
+			PrivateDependencyModuleNames.AddRange(
+				new string[]
+				{
+					"UnrealEd",
+					"SettingsEditor",
+					"Slate",
+					"SlateCore",
+				}
+			);
+		}
+	}
+}
diff --git a/Engine/Plugins/Experimental/Compositing/CompositeCore/Source/CompositeCore/Private/CompositeCoreModule.cpp b/Engine/Plugins/Experimental/Compositing/CompositeCore/Source/CompositeCore/Private/CompositeCoreModule.cpp
new file mode 100644
--- /dev/null
+++ b/Engine/Plugins/Experimental/Compositing/CompositeCore/Source/CompositeCore/Private/CompositeCoreModule.cpp
@@ -0,0 +1,31 @@
+// Copyright Epic Games, Inc. All Rights Reserved.
+
+#include "CompositeCoreModule.h"
+
+#include "Interfaces/IPluginManager.h"
+#include "HDRHelper.h"
+#include "Misc/Paths.h"
+#include "Modules/ModuleManager.h"
+#include "ShaderCore.h"
+
+DEFINE_LOG_CATEGORY(LogCompositeCore);
+
+#define LOCTEXT_NAMESPACE "CompositeCore"
+
+void FCompositeCoreModule::StartupModule()
+{
+	AddShaderSourceDirectoryMapping(TEXT("/Plugin/CompositeCore"), FPaths::Combine(IPluginManager::Get().FindPlugin(TEXT("CompositeCore"))->GetBaseDir(), TEXT("Shaders")));
+
+	if (IsHDREnabled())
+	{
+		UE_LOG(LogCompositeCore, Warning, TEXT("Composite pipeline disabled: HDR mode is not currently supported."));
+	}
+}
+
+void FCompositeCoreModule::ShutdownModule()
+{
+}
+
+IMPLEMENT_MODULE( FCompositeCoreModule, CompositeCore )
+
+#undef LOCTEXT_NAMESPACE
diff --git a/Engine/Plugins/Experimental/Compositing/CompositeCore/Source/CompositeCore/Private/CompositeCoreSceneViewExtension.cpp b/Engine/Plugins/Experimental/Compositing/CompositeCore/Source/CompositeCore/Private/CompositeCoreSceneViewExtension.cpp
new file mode 100644
--- /dev/null
+++ b/Engine/Plugins/Experimental/Compositing/CompositeCore/Source/CompositeCore/Private/CompositeCoreSceneViewExtension.cpp
@@ -0,0 +1,621 @@
+// Copyright Epic Games, Inc. All Rights Reserved.
+
+#include "CompositeCoreSceneViewExtension.h"
+
+#include "CompositeCoreModule.h"
+
+#include "CommonRenderResources.h"
+#include "Components/PrimitiveComponent.h"
+#include "Containers/Map.h"
+#include "Containers/Set.h"
+#include "EngineUtils.h"
+#include "Engine/Texture.h"
+#include "FXRenderingUtils.h"
+#include "HDRHelper.h"
+#include "Materials/Material.h"
+#include "PixelShaderUtils.h"
+#include "PostProcess/LensDistortion.h"
+#include "PostProcess/PostProcessMaterialInputs.h"
+#include "Rendering/CustomRenderPass.h"
+#include "RenderGraphBuilder.h"
+#include "RenderGraphUtils.h"
+#include "SceneView.h"
+#include "ScreenPass.h"
+#include "SystemTextures.h"
+#include "TextureResource.h"
+
+DECLARE_GPU_STAT_NAMED(FCompositeCoreDilate, TEXT("CompositeCore.Dilate"));
+
+class FCompositeCoreDilateShader : public FGlobalShader
+{
+	DECLARE_GLOBAL_SHADER(FCompositeCoreDilateShader);
+	SHADER_USE_PARAMETER_STRUCT(FCompositeCoreDilateShader, FGlobalShader);
+
+	BEGIN_SHADER_PARAMETER_STRUCT(FParameters, )
+		SHADER_PARAMETER_RDG_TEXTURE(Texture2D, InputTexture)
+		SHADER_PARAMETER_RDG_TEXTURE_UAV(RWTexture2D<float4>, RWOutputTexture)
+		SHADER_PARAMETER(FIntPoint, Dimensions)
+	END_SHADER_PARAMETER_STRUCT()
+
+	static inline void ModifyCompilationEnvironment(const FGlobalShaderPermutationParameters& Parameters, FShaderCompilerEnvironment& OutEnvironment)
+	{
+		OutEnvironment.SetDefine(TEXT("TILE_SIZE"), ThreadGroupSize);
+
+		FGlobalShader::ModifyCompilationEnvironment(Parameters, OutEnvironment);
+	}
+
+	static const uint32 ThreadGroupSize = 16;
+};
+IMPLEMENT_GLOBAL_SHADER(FCompositeCoreDilateShader, "/Plugin/CompositeCore/Private/CompositeCoreDilate.usf", "MainCS", SF_Compute);
+
+namespace UE
+{
+	namespace CompositeCore
+	{
+		namespace Private
+		{
+			void AddDilatePass(FRDGBuilder& GraphBuilder, FRDGTextureRef Input, FRDGTextureRef Output, ERHIFeatureLevel::Type FeatureLevel)
+			{
+				RDG_EVENT_SCOPE_STAT(GraphBuilder, FCompositeCoreDilate, "CompositeCore.Dilate");
+				RDG_GPU_STAT_SCOPE(GraphBuilder, FCompositeCoreDilate);
+
+				FGlobalShaderMap* GlobalShaderMap = GetGlobalShaderMap(FeatureLevel);
+				const FIntPoint TextureSize = Input->Desc.Extent;
+
+				FCompositeCoreDilateShader::FParameters* PassParameters = GraphBuilder.AllocParameters<FCompositeCoreDilateShader::FParameters>();
+				PassParameters->InputTexture = Input;
+				PassParameters->RWOutputTexture = GraphBuilder.CreateUAV(Output);
+				PassParameters->Dimensions = TextureSize;
+
+				TShaderMapRef<FCompositeCoreDilateShader> ComputeShader(GlobalShaderMap);
+				FComputeShaderUtils::AddPass(
+					GraphBuilder,
+					RDG_EVENT_NAME("CompositeCore.Dilate (%dx%d)", TextureSize.X, TextureSize.Y),
+					GSupportsEfficientAsyncCompute ? ERDGPassFlags::AsyncCompute : ERDGPassFlags::Compute,
+					ComputeShader,
+					PassParameters,
+					FComputeShaderUtils::GetGroupCount(TextureSize, FCompositeCoreDilateShader::ThreadGroupSize)
+				);
+			}
+
+			FRDGTextureDesc GetPostProcessingDesc(FRDGTextureDesc InDesc, EPixelFormat InOutputFormat)
+			{
+				FRDGTextureDesc OutputDesc = InDesc;
+				OutputDesc.Dimension = ETextureDimension::Texture2D;
+				OutputDesc.ArraySize = 1;
+
+				OutputDesc.Reset();
+				if (InOutputFormat != PF_Unknown)
+				{
+					OutputDesc.Format = InOutputFormat;
+				}
+				OutputDesc.ClearValue = FClearValueBinding(FLinearColor::Black);
+				OutputDesc.Flags &= (~ETextureCreateFlags::FastVRAM);
+
+				return OutputDesc;
+			}
+		}
+	}
+}
+
+class FCompositeCoreCustomRenderPass : public FCustomRenderPassBase
+{
+public:
+	IMPLEMENT_CUSTOM_RENDER_PASS(FCompositeCoreCustomRenderPass);
+
+	FCompositeCoreCustomRenderPass(const FIntPoint& InRenderTargetSize, FCompositeCoreSceneViewExtension* InParentExtension, FSceneView& InView, int32 InDilationSize)
+		: FCustomRenderPassBase(TEXT("CompositeCoreCustomRenderPass"), ERenderMode::DepthAndBasePass, ERenderOutput::SceneColorAndAlpha, InRenderTargetSize)
+		, ParentExtension(InParentExtension)
+		, ViewId(InView.GetViewKey())
+		, ViewFeatureLevel(InView.GetFeatureLevel())
+		, DilationSize(InDilationSize)
+	{
+		bSceneColorWithTranslucent = true;
+	}
+
+	virtual void OnPreRender(FRDGBuilder& GraphBuilder) override
+	{
+		const FRDGTextureDesc TextureDesc = FRDGTextureDesc::Create2D(RenderTargetSize, PF_FloatRGBA, FClearValueBinding::Black, TexCreate_RenderTargetable | TexCreate_ShaderResource);
+		RenderTargetTexture = GraphBuilder.CreateTexture(TextureDesc, TEXT("CompositeCoreCustomTexture"));
+		AddClearRenderTargetPass(GraphBuilder, RenderTargetTexture, FLinearColor::Black, FIntRect(FInt32Point(), RenderTargetSize));
+	}
+
+	virtual void OnPostRender(FRDGBuilder& GraphBuilder) override
+	{
+		const FRDGTextureDesc TextureDesc = FRDGTextureDesc::Create2D(RenderTargetTexture->Desc.Extent, PF_FloatRGBA, FClearValueBinding::Black, TexCreate_UAV | TexCreate_ShaderResource | TexCreate_RenderTargetable);
+		FRDGTextureRef Output = GraphBuilder.CreateTexture(TextureDesc, TEXT("CompositeCoreProcessedTexture"));
+
+		// TODO: Add support for wider dilation
+		if (DilationSize > 0)
+		{
+			UE::CompositeCore::Private::AddDilatePass(GraphBuilder, RenderTargetTexture, Output, ViewFeatureLevel);
+		}
+
+		ParentExtension->CollectCustomRenderTarget(ViewId, GraphBuilder.ConvertToExternalTexture(Output));
+	}
+
+private:
+
+	FCompositeCoreSceneViewExtension* ParentExtension;
+	const uint32 ViewId;
+	const ERHIFeatureLevel::Type ViewFeatureLevel;
+	const int32 DilationSize;
+};
+
+FCompositeCoreSceneViewExtension::FCompositeCoreSceneViewExtension(const FAutoRegister& AutoReg, UWorld* InWorld)
+	: FWorldSceneViewExtension(AutoReg, InWorld)
+{
+	BuiltInRenderPassOptions = UE::CompositeCore::FBuiltInRenderPassOptions{};
+}
+
+FCompositeCoreSceneViewExtension::~FCompositeCoreSceneViewExtension() = default;
+
+void FCompositeCoreSceneViewExtension::RegisterPrimitives_GameThread(const TArray<UPrimitiveComponent*>& InPrimitiveComponents)
+{
+	check(IsInGameThread());
+
+	for (UPrimitiveComponent* InPrimitiveComponent : InPrimitiveComponents)
+	{
+		if (!IsValid(InPrimitiveComponent))
+		{
+			continue;
+		}
+
+		if (!CompositePrimitives.Contains(InPrimitiveComponent))
+		{
+			CompositePrimitives.Add(InPrimitiveComponent);
+		}
+
+		// The SetHoldout() function makes changes only if the HoldoutState value differs from the current value.
+		InPrimitiveComponent->SetHoldout(true);
+	}
+}
+
+void FCompositeCoreSceneViewExtension::UnregisterPrimitives_GameThread(const TArray<UPrimitiveComponent*>& InPrimitiveComponents)
+{
+	check(IsInGameThread());
+
+	for (UPrimitiveComponent* InPrimitiveComponent : InPrimitiveComponents)
+	{
+		if (!IsValid(InPrimitiveComponent))
+		{
+			continue;
+		}
+
+		if (CompositePrimitives.Contains(InPrimitiveComponent))
+		{
+			CompositePrimitives.Remove(InPrimitiveComponent);
+		}
+
+		// The SetHoldout() function makes changes only if the HoldoutState value differs from the current value.
+		InPrimitiveComponent->SetHoldout(false);
+	}
+}
+
+void FCompositeCoreSceneViewExtension::SetCompositeTextureOverride_GameThread(UTexture* InTextureOverride)
+{
+	CompositeTextureOverride = InTextureOverride;
+}
+
+void FCompositeCoreSceneViewExtension::SetPostRenderWork_GameThread(UE::CompositeCore::FPostRenderWork&& InWork)
+{
+	ENQUEUE_RENDER_COMMAND(CopyCompositeCoreRenderWork)(
+		[InWork = MoveTemp(InWork), WeakThis = AsWeak()](FRHICommandListImmediate& RHICmdList) mutable
+		{
+			TSharedPtr<FCompositeCoreSceneViewExtension> SVE = StaticCastSharedPtr<FCompositeCoreSceneViewExtension>(WeakThis.Pin());
+			if (SVE.IsValid())
+			{
+				SVE->PostRenderWork_RenderThread = MoveTemp(InWork);
+			}
+		});
+}
+
+void FCompositeCoreSceneViewExtension::ResetPostRenderWork_GameThread()
+{
+	ENQUEUE_RENDER_COMMAND(CopyCompositeCoreRenderWork)(
+		[WeakThis = AsWeak()](FRHICommandListImmediate& RHICmdList)
+		{
+			TSharedPtr<FCompositeCoreSceneViewExtension> SVE = StaticCastSharedPtr<FCompositeCoreSceneViewExtension>(WeakThis.Pin());
+			if (SVE.IsValid())
+			{
+				SVE->PostRenderWork_RenderThread.Reset();
+			}
+		});
+}
+
+void FCompositeCoreSceneViewExtension::SetBuiltInRenderPassOptions_GameThread(const UE::CompositeCore::FBuiltInRenderPassOptions& InOptions)
+{
+	BuiltInRenderPassOptions = InOptions;
+}
+
+bool FCompositeCoreSceneViewExtension::IsActiveThisFrame_Internal(const FSceneViewExtensionContext& Context) const
+{
+	bool bIsActive = FWorldSceneViewExtension::IsActiveThisFrame_Internal(Context);
+	bIsActive &= !CompositePrimitives.IsEmpty();
+	bIsActive &= !IsHDREnabled();
+
+	return bIsActive;
+}
+
+
+//~ Begin ISceneViewExtension Interface
+
+int32 FCompositeCoreSceneViewExtension::GetPriority() const
+{
+	return GetDefault<UCompositeCorePluginSettings>()->SceneViewExtensionPriority;
+}
+
+const UE::CompositeCore::FPostRenderWork& FCompositeCoreSceneViewExtension::GetRenderWork() const
+{
+	if (PostRenderWork_RenderThread.IsSet())
+	{
+		return PostRenderWork_RenderThread.GetValue();
+	}
+	else
+	{
+		return UE::CompositeCore::FPostRenderWork::GetDefault();
+	}
+}
+
+void FCompositeCoreSceneViewExtension::SetupViewFamily(FSceneViewFamily& InViewFamily)
+{
+	// Cleanup invalid primitives.
+	for (auto Iter = CompositePrimitives.CreateIterator(); Iter; ++Iter)
+	{
+		if (!Iter->IsValid())
+		{
+			Iter.RemoveCurrent();
+		}
+	}
+}
+
+void FCompositeCoreSceneViewExtension::SetupView(FSceneViewFamily& InViewFamily, FSceneView& InView)
+{
+	// Disable custom render pass for view families without primitive alpha holdout support.
+	if (InViewFamily.EngineShowFlags.AllowPrimitiveAlphaHoldout == 0)
+	{
+		return;
+	}
+
+	const TWeakObjectPtr<UWorld> WorldPtr = GetWorld();
+	check(WorldPtr.IsValid());
+
+	TSet<FPrimitiveComponentId> CompositeCorePrimitiveIds;
+	for (const TWeakObjectPtr<UPrimitiveComponent>& PrimitivePtr : CompositePrimitives)
+	{
+		TStrongObjectPtr<UPrimitiveComponent> Primitive = PrimitivePtr.Pin();
+		// Collect only those primitives that use the bHoldout flag
+		// The user can directly change this flag outside of this VE.
+		if (Primitive.IsValid() && Primitive->bHoldout)
+		{
+			const FPrimitiveComponentId PrimId = Primitive->GetPrimitiveSceneId();
+
+			if (InView.ShowOnlyPrimitives.IsSet())
+			{
+				if (InView.ShowOnlyPrimitives.GetValue().Contains(PrimId))
+				{
+					CompositeCorePrimitiveIds.Add(Primitive->GetPrimitiveSceneId());
+				}
+			}
+			else if (!InView.HiddenPrimitives.Contains(PrimId))
+			{
+				CompositeCorePrimitiveIds.Add(Primitive->GetPrimitiveSceneId());
+			}
+		}
+	}
+
+	if (CompositeCorePrimitiveIds.IsEmpty())
+	{
+		return;
+	}
+
+	// Extract the custom render target size
+	FIntPoint RenderTargetViewSize;
+	if (InView.PrimaryScreenPercentageMethod == EPrimaryScreenPercentageMethod::TemporalUpscale)
+	{
+		// Note: This is equivalent to FViewInfo::GetSecondaryViewRectSize
+		FIntRect OutputRect = FIntRect(0, 0,
+			FMath::CeilToInt(InView.UnscaledViewRect.Width() * InViewFamily.SecondaryViewFraction),
+			FMath::CeilToInt(InView.UnscaledViewRect.Height() * InViewFamily.SecondaryViewFraction)
+		);
+		QuantizeSceneBufferSize(OutputRect.Max, RenderTargetViewSize);
+	}
+	else
+	{
+		RenderTargetViewSize = InView.UnscaledViewRect.Size();
+	}
+
+	// Create a new custom render pass to render the composite primitive(s)
+	FCompositeCoreCustomRenderPass* CustomRenderPass = new FCompositeCoreCustomRenderPass(
+		RenderTargetViewSize,
+		this,
+		InView,
+		BuiltInRenderPassOptions.OutputDilationSize
+	);
+
+	FSceneInterface::FCustomRenderPassRendererInput PassInput{};
+	PassInput.EngineShowFlags = InViewFamily.EngineShowFlags;
+	PassInput.EngineShowFlags.DisableFeaturesForUnlit();
+	PassInput.EngineShowFlags.SetTranslucency(true);
+	PassInput.EngineShowFlags.SetUnlitViewmode(BuiltInRenderPassOptions.bEnableUnlitViewmode);
+	PassInput.EngineShowFlags.SetAllowPrimitiveAlphaHoldout(false);
+	if (BuiltInRenderPassOptions.ViewUserFlagsOverride.IsSet())
+	{
+		PassInput.bOverridesPostVolumeUserFlags = true;
+		PassInput.PostVolumeUserFlags = BuiltInRenderPassOptions.ViewUserFlagsOverride.GetValue();
+	}
+	// Note: Incoming view location is invalid for scene captures
+	PassInput.ViewLocation = InView.ViewMatrices.GetViewOrigin();
+	PassInput.ViewRotationMatrix = InView.ViewMatrices.GetViewMatrix().RemoveTranslation();
+	PassInput.ViewRotationMatrix.RemoveScaling();
+
+	// Note: Projection matrix here is without jitter, GetProjectionNoAAMatrix() is invalid (not yet available).
+	PassInput.ProjectionMatrix = InView.ViewMatrices.GetProjectionMatrix();
+	PassInput.ViewActor = InView.ViewActor;
+	PassInput.ShowOnlyPrimitives = CompositeCorePrimitiveIds;
+	PassInput.CustomRenderPass = CustomRenderPass;
+	PassInput.bIsSceneCapture = true;
+
+	WorldPtr.Get()->Scene->AddCustomRenderPass(&InViewFamily, PassInput);
+}
+
+void FCompositeCoreSceneViewExtension::PreRenderViewFamily_RenderThread(FRDGBuilder& GraphBuilder, FSceneViewFamily& InViewFamily)
+{
+}
+
+void FCompositeCoreSceneViewExtension::PreRenderView_RenderThread(FRDGBuilder& GraphBuilder, FSceneView& InView)
+{
+}
+
+void FCompositeCoreSceneViewExtension::PrePostProcessPass_RenderThread(FRDGBuilder& GraphBuilder, const FSceneView& InView, const FPostProcessingInputs& Inputs)
+{
+}
+
+void FCompositeCoreSceneViewExtension::SubscribeToPostProcessingPass(EPostProcessingPass PassId, const FSceneView& InView, FAfterPassCallbackDelegateArray& InOutPassCallbacks, bool bIsPassEnabled)
+{
+	if (!IsActiveForView(InView))
+	{
+		return;
+	}
+	
+	if (PassId == EPostProcessingPass::BeforeDOF)
+	{
+		// We use BeforeDOF as an input processing step for simplicity, since PrePostProcessPass_RenderThread has a different interface.
+		InOutPassCallbacks.Add(FAfterPassCallbackDelegate::CreateRaw(this, &FCompositeCoreSceneViewExtension::InputProcessing_RenderThread));
+	}
+	if (PassId == EPostProcessingPass::SSRInput && GetRenderWork().FramePasses.Contains(ISceneViewExtension::EPostProcessingPass::SSRInput))
+	{
+		InOutPassCallbacks.Add(FAfterPassCallbackDelegate::CreateRaw(this, &FCompositeCoreSceneViewExtension::SSRInputWork_RenderThread));
+	}
+	else if (PassId == EPostProcessingPass::Tonemap && GetRenderWork().FramePasses.Contains(ISceneViewExtension::EPostProcessingPass::Tonemap))
+	{
+		InOutPassCallbacks.Add(FAfterPassCallbackDelegate::CreateRaw(this, &FCompositeCoreSceneViewExtension::TonemapWork_RenderThread));
+	}
+	else if (PassId == EPostProcessingPass::VisualizeDepthOfField && GetRenderWork().FramePasses.Contains(ISceneViewExtension::EPostProcessingPass::VisualizeDepthOfField))
+	{
+		InOutPassCallbacks.Add(FAfterPassCallbackDelegate::CreateRaw(this, &FCompositeCoreSceneViewExtension::EndOfPostProcessing_RenderThread));
+	}
+}
+
+bool FCompositeCoreSceneViewExtension::IsActiveForView(const FSceneView& InView) const
+{
+	constexpr bool bEvenIfPendingKill = false;
+	constexpr bool bThreadSafeTest = true;
+
+	bool bIsActive = InView.Family->bIsMainViewFamily;
+	bIsActive &= !InView.bIsSceneCapture;
+	bIsActive &= static_cast<bool>(InView.Family->EngineShowFlags.AllowPrimitiveAlphaHoldout);
+	bIsActive &= CustomRenderTargetPerView_RenderThread.Contains(InView.GetViewKey()) || CompositeTextureOverride.IsValid(bEvenIfPendingKill, bThreadSafeTest);
+
+	return bIsActive;
+}
+
+FScreenPassTexture FCompositeCoreSceneViewExtension::InputProcessing_RenderThread(FRDGBuilder& GraphBuilder, const FSceneView& InView, const FPostProcessMaterialInputs& Inputs)
+{
+	using namespace UE::CompositeCore::Private;
+
+	// Reset composite textures for view.
+	// NOTE: all post-processing SVE callbacks occur on *the same view*, so this is allowed.
+	CompositeRenderPass_RenderThread = CompositeTexture_RenderThread = FScreenPassTexture();
+
+	// Prepare texture inputs for following post-processing passes
+	{
+		const TRefCountPtr<IPooledRenderTarget>* CompositeRenderPassPtr = CustomRenderTargetPerView_RenderThread.Find(InView.GetViewKey());
+		if (CompositeRenderPassPtr != nullptr)
+		{
+			CompositeRenderPass_RenderThread = FScreenPassTexture{ GraphBuilder.RegisterExternalTexture(*CompositeRenderPassPtr) };
+		}
+
+		constexpr bool bEvenIfPendingKill = false;
+		constexpr bool bThreadSafeTest = true;
+		if (CompositeTextureOverride.IsValid(bEvenIfPendingKill, bThreadSafeTest))
+		{
+			if (FTextureResource* TextureResource = CompositeTextureOverride.GetEvenIfUnreachable()->GetResource())
+			{
+				const FTextureRHIRef& TextureRHI = TextureResource->GetTextureRHI();
+				if (TextureRHI.IsValid())
+				{
+					FRDGTextureRef CompositeTextureRDG = GraphBuilder.FindExternalTexture(TextureRHI.GetReference());
+
+					if (CompositeTextureRDG == nullptr)
+					{
+						CompositeTextureRDG = GraphBuilder.RegisterExternalTexture(CreateRenderTarget(TextureRHI.GetReference(), TEXT("CompositeTextureOverride")));
+					}
+					
+					FScreenPassTexture Input{ CompositeTextureRDG };
+					
+					const FScreenPassTextureSlice SceneColorSlice = Inputs.GetInput(EPostProcessMaterialInput::SceneColor);
+					const FScreenPassTextureViewport SceneColorViewport = FScreenPassTextureViewport(SceneColorSlice);
+
+					FRDGTextureDesc OutputDesc = GetPostProcessingDesc(SceneColorSlice.TextureSRV->Desc.Texture->Desc, Inputs.OutputFormat);
+					FScreenPassRenderTarget Output = FScreenPassRenderTarget(GraphBuilder.CreateTexture(OutputDesc, TEXT("CompositeTextureCopy")), SceneColorViewport.Rect, InView.GetOverwriteLoadAction());
+	
+					// TODO: We could avoid the copy if the texture override format is usable as scene color.
+					{
+						const FScreenPassTextureViewport InputViewport(Input);
+						const FScreenPassTextureViewport OutputViewport(Output);
+
+						TShaderMapRef<FCopyRectPS> PixelShader(GetGlobalShaderMap(InView.GetFeatureLevel()));
+
+						FCopyRectPS::FParameters* Parameters = GraphBuilder.AllocParameters<FCopyRectPS::FParameters>();
+						Parameters->InputTexture = Input.Texture;
+						Parameters->InputSampler = TStaticSamplerState<SF_Bilinear>::GetRHI();
+						Parameters->RenderTargets[0] = Output.GetRenderTargetBinding();
+
+						AddDrawScreenPass(GraphBuilder, RDG_EVENT_NAME("DrawTexture"), InView, OutputViewport, InputViewport, PixelShader, Parameters);
+					}
+
+					CompositeTexture_RenderThread = Output;
+				}
+			}
+		}
+	}
+
+	// Texture processing
+	for (TPair<ECompositeCorePassInput, TArray<const ICompositeCorePassProxy*>> Processing : GetRenderWork().InputProcessing)
+	{
+		const ECompositeCorePassInput ActiveTextureInput = Processing.Key;
+		const TArray<const ICompositeCorePassProxy*>& Passes = Processing.Value;
+
+		if (!ensureMsgf(ActiveTextureInput != ECompositeCorePassInput::Default, TEXT("Processing is expected to run on composite textures.")))
+		{
+			continue;
+		}
+
+		TSortedMap<ECompositeCorePassInput, FScreenPassTextureSlice> TextureInputs;
+		if (CompositeRenderPass_RenderThread.IsValid())
+		{
+			TextureInputs.Emplace(ECompositeCorePassInput::CompositeRenderPass, FScreenPassTextureSlice::CreateFromScreenPassTexture(GraphBuilder, CompositeRenderPass_RenderThread));
+		}
+		if (CompositeTexture_RenderThread.IsValid())
+		{
+			TextureInputs.Emplace(ECompositeCorePassInput::CompositeTexture, FScreenPassTextureSlice::CreateFromScreenPassTexture(GraphBuilder, CompositeTexture_RenderThread));
+		}
+
+		FScreenPassTexture Output;
+
+		for (const ICompositeCorePassProxy* Pass : Passes)
+		{
+			if (FScreenPassTextureSlice* SliceInput = TextureInputs.Find(ActiveTextureInput))
+			{
+				FPostProcessMaterialInputs PassInputs = Inputs;
+				PassInputs.SetInput(EPostProcessMaterialInput::SceneColor, *SliceInput);
+				PassInputs.OverrideOutput = FScreenPassRenderTarget{};
+
+				// Register pass and update output
+				Output = Pass->Add(GraphBuilder, InView, PassInputs, ISceneViewExtension::EPostProcessingPass::BeforeDOF);
+
+				// Update active input for the next pass
+				TextureInputs[ActiveTextureInput] = FScreenPassTextureSlice::CreateFromScreenPassTexture(GraphBuilder, Output);
+			}
+		}
+
+		// Update composite textures for later callbacks in post-processing
+		if (Output.IsValid())
+		{
+			switch (ActiveTextureInput)
+			{
+			case ECompositeCorePassInput::CompositeRenderPass:
+				CompositeRenderPass_RenderThread = Output;
+				break;
+			case ECompositeCorePassInput::CompositeTexture:
+				CompositeTexture_RenderThread = Output;
+				break;
+
+			default:
+				checkNoEntry();
+				break;
+			}
+		}
+	}
+
+	return Inputs.ReturnUntouchedSceneColorForPostProcessing(GraphBuilder);
+}
+
+FScreenPassTexture FCompositeCoreSceneViewExtension::PostProcessWork_RenderThread(FRDGBuilder& GraphBuilder, const FSceneView& InView, const FPostProcessMaterialInputs& Inputs, ISceneViewExtension::EPostProcessingPass InLocation)
+{
+	using namespace UE::CompositeCore::Private;
+
+	TSortedMap<ECompositeCorePassInput, FScreenPassTextureSlice> TextureInputs;
+	TextureInputs.Emplace(ECompositeCorePassInput::Default, Inputs.GetInput(EPostProcessMaterialInput::SceneColor));
+	if (CompositeRenderPass_RenderThread.IsValid())
+	{
+		TextureInputs.Emplace(ECompositeCorePassInput::CompositeRenderPass, FScreenPassTextureSlice::CreateFromScreenPassTexture(GraphBuilder, CompositeRenderPass_RenderThread));
+	}
+	if (CompositeTexture_RenderThread.IsValid())
+	{
+		TextureInputs.Emplace(ECompositeCorePassInput::CompositeTexture, FScreenPassTextureSlice::CreateFromScreenPassTexture(GraphBuilder, CompositeTexture_RenderThread));
+	}
+
+	FScreenPassTexture Output;
+	const TArray<const ICompositeCorePassProxy*>& Passes = GetRenderWork().FramePasses[InLocation];
+	for (int32 Index = 0; Index < Passes.Num(); ++Index)
+	{
+		const ICompositeCorePassProxy* Pass = Passes[Index];
+		const bool bIsLastPass = (Index == Passes.Num() - 1);
+
+		// Update pass inputs
+		FPostProcessMaterialInputs PassInputs = Inputs;
+		int32 ConnectedInputs = 0;
+		for (int32 InputIndex = 0; InputIndex < Pass->GetNumInputs() && InputIndex < TextureInputs.Num(); ++InputIndex)
+		{
+			if (FScreenPassTextureSlice* SliceInput = TextureInputs.Find(Pass->GetPassInput(InputIndex)))
+			{
+				PassInputs.SetInput(static_cast<EPostProcessMaterialInput>(InputIndex), *SliceInput);
+				ConnectedInputs++;
+			}
+		}
+
+		if (ConnectedInputs == 0)
+		{
+			continue;
+		}
+
+		// Only apply the output override on the last pass
+		if (!bIsLastPass)
+		{
+			PassInputs.OverrideOutput = FScreenPassRenderTarget{};
+		}
+
+		// Register pass and update output
+		Output = Pass->Add(GraphBuilder, InView, PassInputs, InLocation);
+
+		// Update input 0 for the next pass
+		TextureInputs[ECompositeCorePassInput::Default] = FScreenPassTextureSlice::CreateFromScreenPassTexture(GraphBuilder, Output);
+	}
+
+	if (Output.IsValid())
+	{
+		return MoveTemp(Output);
+	}
+	else
+	{
+		return Inputs.ReturnUntouchedSceneColorForPostProcessing(GraphBuilder);
+	}
+}
+
+FScreenPassTexture FCompositeCoreSceneViewExtension::SSRInputWork_RenderThread(FRDGBuilder& GraphBuilder, const FSceneView& InView, const FPostProcessMaterialInputs& Inputs)
+{
+	return PostProcessWork_RenderThread(GraphBuilder, InView, Inputs, ISceneViewExtension::EPostProcessingPass::SSRInput);
+}
+
+FScreenPassTexture FCompositeCoreSceneViewExtension::TonemapWork_RenderThread(FRDGBuilder& GraphBuilder, const FSceneView& InView, const FPostProcessMaterialInputs& Inputs)
+{
+	return PostProcessWork_RenderThread(GraphBuilder, InView, Inputs, ISceneViewExtension::EPostProcessingPass::Tonemap);
+}
+
+FScreenPassTexture FCompositeCoreSceneViewExtension::EndOfPostProcessing_RenderThread(FRDGBuilder& GraphBuilder, const FSceneView& InView, const FPostProcessMaterialInputs& Inputs)
+{
+	return PostProcessWork_RenderThread(GraphBuilder, InView, Inputs, ISceneViewExtension::EPostProcessingPass::VisualizeDepthOfField);
+}
+
+void FCompositeCoreSceneViewExtension::PostRenderViewFamily_RenderThread(FRDGBuilder& GraphBuilder, FSceneViewFamily& InViewFamily)
+{
+
+}
+
+void FCompositeCoreSceneViewExtension::PostRenderView_RenderThread(FRDGBuilder& GraphBuilder, FSceneView& InView)
+{
+	CustomRenderTargetPerView_RenderThread.Remove(InView.GetViewKey());
+}
+
diff --git a/Engine/Plugins/Experimental/Compositing/CompositeCore/Source/CompositeCore/Private/CompositeCoreSceneViewExtension.h b/Engine/Plugins/Experimental/Compositing/CompositeCore/Source/CompositeCore/Private/CompositeCoreSceneViewExtension.h
new file mode 100644
--- /dev/null
+++ b/Engine/Plugins/Experimental/Compositing/CompositeCore/Source/CompositeCore/Private/CompositeCoreSceneViewExtension.h
@@ -0,0 +1,104 @@
+// Copyright Epic Games, Inc. All Rights Reserved.
+
+#pragma once
+
+#include "SceneViewExtension.h"
+
+#include "Containers/ContainersFwd.h"
+#include "RenderGraphFwd.h"
+#include "RendererInterface.h"
+#include "ScreenPass.h"
+
+#include "CompositeCoreSettings.h"
+#include "Passes/CompositeCorePassProxy.h"
+
+class FViewInfo;
+struct FScreenPassRenderTarget;
+struct FScreenPassTexture;
+class UMaterialInterface;
+class UTexture;
+
+class FCompositeCoreSceneViewExtension : public FWorldSceneViewExtension
+{
+public:
+	FCompositeCoreSceneViewExtension(const FAutoRegister& AutoReg, UWorld* InWorld);
+	~FCompositeCoreSceneViewExtension();
+
+	/* Register primitives for compositing. */
+	void RegisterPrimitives_GameThread(const TArray<UPrimitiveComponent*>& InPrimitiveComponents);
+
+	/* Unregister primitives for compositing. */
+	void UnregisterPrimitives_GameThread(const TArray<UPrimitiveComponent*>& InPrimitiveComponents);
+
+	/* Set a texture override to be used instead of the separate composited render. */
+	void SetCompositeTextureOverride_GameThread(UTexture* InTextureOverride);
+
+	/* Set post-processing render work to be enqueued on the render thread. */
+	void SetPostRenderWork_GameThread(UE::CompositeCore::FPostRenderWork&& InWork);
+
+	/* Reset post-processing render work. */
+	void ResetPostRenderWork_GameThread();
+
+	/* Set an optional view user flags override. */
+	void SetBuiltInRenderPassOptions_GameThread(const UE::CompositeCore::FBuiltInRenderPassOptions& InOptions);
+
+	/* Called by the custom render pass to store its view render target for this frame. */
+	template<typename T>
+	void CollectCustomRenderTarget(uint32 InViewId, T InRenderTarget)
+	{
+		CustomRenderTargetPerView_RenderThread.Add(InViewId, InRenderTarget);
+	}
+
+	//~ Begin ISceneViewExtension Interface
+	virtual int32 GetPriority() const override;
+	virtual void SetupViewFamily(FSceneViewFamily& InViewFamily) override;
+	virtual void SetupView(FSceneViewFamily& InViewFamily, FSceneView& InView) override;
+	virtual void PreRenderViewFamily_RenderThread(FRDGBuilder& GraphBuilder, FSceneViewFamily& InViewFamily) override;
+	virtual void PreRenderView_RenderThread(FRDGBuilder& GraphBuilder, FSceneView& InView) override;
+	virtual void PrePostProcessPass_RenderThread(FRDGBuilder& GraphBuilder, const FSceneView& InView, const FPostProcessingInputs& Inputs) override;
+	virtual void SubscribeToPostProcessingPass(EPostProcessingPass PassId, const FSceneView& InView, FAfterPassCallbackDelegateArray& InOutPassCallbacks, bool bIsPassEnabled) override;
+	virtual void PostRenderViewFamily_RenderThread(FRDGBuilder& GraphBuilder, FSceneViewFamily& InViewFamily) override;
+	virtual void PostRenderView_RenderThread(FRDGBuilder& GraphBuilder, FSceneView& InView) override;
+	//~ End ISceneViewExtension Interface
+
+protected:
+	//~ Begin ISceneViewExtension Interface
+	virtual bool IsActiveThisFrame_Internal(const FSceneViewExtensionContext& Context) const override;
+	//~ End ISceneViewExtension Interface
+
+private:
+	/* Getter for the frame render work, either from the actor or default. */
+	const UE::CompositeCore::FPostRenderWork& GetRenderWork() const;
+
+	/* Active check to avoid work in post-processing. */
+	bool IsActiveForView(const FSceneView& InView) const;
+
+	/* Callbacks for processing passes. */
+	FScreenPassTexture InputProcessing_RenderThread(FRDGBuilder& GraphBuilder, const FSceneView& InView, const FPostProcessMaterialInputs& Inputs);
+	FScreenPassTexture PostProcessWork_RenderThread(FRDGBuilder& GraphBuilder, const FSceneView& InView, const FPostProcessMaterialInputs& Inputs, ISceneViewExtension::EPostProcessingPass InLocation);
+	FScreenPassTexture SSRInputWork_RenderThread(FRDGBuilder& GraphBuilder, const FSceneView& InView, const FPostProcessMaterialInputs& Inputs);
+	FScreenPassTexture TonemapWork_RenderThread(FRDGBuilder& GraphBuilder, const FSceneView& InView, const FPostProcessMaterialInputs& Inputs);
+	FScreenPassTexture EndOfPostProcessing_RenderThread(FRDGBuilder& GraphBuilder, const FSceneView& InView, const FPostProcessMaterialInputs& Inputs);
+
+	/* Collection of primitives to render as a custom render pass and composite after post-processing. .*/
+	TSet<TWeakObjectPtr<UPrimitiveComponent>> CompositePrimitives;
+
+	/* Built-in composite render pass options. .*/
+	UE::CompositeCore::FBuiltInRenderPassOptions BuiltInRenderPassOptions;
+
+	/* Composite texture override use directly instead of the separate composited render. */
+	TWeakObjectPtr<UTexture> CompositeTextureOverride;
+
+	/* Custom render pass render targets for each active view .*/
+	TMap<uint32, TRefCountPtr<IPooledRenderTarget>> CustomRenderTargetPerView_RenderThread;
+
+	/* Custom actor render work. */
+	TOptional<UE::CompositeCore::FPostRenderWork> PostRenderWork_RenderThread;
+
+	/* Composite render pass screen texture, only valid in post-processing callbacks. */
+	FScreenPassTexture CompositeRenderPass_RenderThread;
+
+	/* Composite screen texture, only valid in post-processing callbacks. */
+	FScreenPassTexture CompositeTexture_RenderThread;
+};
+
diff --git a/Engine/Plugins/Experimental/Compositing/CompositeCore/Source/CompositeCore/Private/CompositeCoreSettings.cpp b/Engine/Plugins/Experimental/Compositing/CompositeCore/Source/CompositeCore/Private/CompositeCoreSettings.cpp
new file mode 100644
--- /dev/null
+++ b/Engine/Plugins/Experimental/Compositing/CompositeCore/Source/CompositeCore/Private/CompositeCoreSettings.cpp
@@ -0,0 +1,97 @@
+// Copyright Epic Games, Inc. All Rights Reserved.
+
+#include "CompositeCoreSettings.h"
+#include "CompositeCoreSubsystem.h"
+
+#include "Components/PrimitiveComponent.h"
+
+#include UE_INLINE_GENERATED_CPP_BY_NAME(CompositeCoreSettings)
+
+namespace UE::CompositeCore
+{
+	static bool GRegisterPrimitivesOnTick = true;
+
+	bool IsRegisterPrimitivesOnTickEnabled()
+	{
+		return GRegisterPrimitivesOnTick;
+	}
+}
+static FAutoConsoleVariableRef CVarRefreshPrimitivesOnTick(
+	TEXT("CompositeCore.RefreshPrimitivesOnTick"),
+	UE::CompositeCore::GRegisterPrimitivesOnTick,
+	TEXT("Refresh primitive registration on every tick. Enabled by default for use cases such as dynamic text 3D.\n")
+	TEXT("Read-only and to be set in a config file (requires restart)."),
+	ECVF_ReadOnly);
+
+UCompositeCorePluginSettings::UCompositeCorePluginSettings()
+	: DisabledPrimitiveClasses(
+		{
+			FSoftClassPath("/Script/Engine.CapsuleComponent"),
+			FSoftClassPath("/Script/Engine.BillboardComponent"),
+			FSoftClassPath("/Script/Engine.ArrowComponent"),
+			FSoftClassPath("/Script/Engine.DrawFrustumComponent"),
+			FSoftClassPath("/Script/Engine.LineBatchComponent")
+		}
+	)
+	, SceneViewExtensionPriority(100 + 1 /* OPENCOLORIO_SCENE_VIEW_EXTENSION_PRIORITY + 1 */)
+{ }
+
+
+FName UCompositeCorePluginSettings::GetCategoryName() const
+{
+	return TEXT("Plugins");
+}
+
+#if WITH_EDITOR
+FText UCompositeCorePluginSettings::GetSectionText() const
+{
+	return NSLOCTEXT("CompositeCoreSettings", "CompositeCoreSettingsSection", "Composite Core");
+}
+
+FName UCompositeCorePluginSettings::GetSectionName() const
+{
+	return TEXT("Composite Core");
+}
+void UCompositeCorePluginSettings::PostEditChangeProperty(FPropertyChangedEvent& PropertyChangedEvent)
+{
+	Super::PostEditChangeProperty(PropertyChangedEvent);
+
+	if (PropertyChangedEvent.GetMemberPropertyName() == GET_MEMBER_NAME_CHECKED(ThisClass, DisabledPrimitiveClasses))
+	{
+		CacheDisabledClasses();
+	}
+}
+#endif
+
+void UCompositeCorePluginSettings::PostInitProperties()
+{
+	Super::PostInitProperties();
+
+	CacheDisabledClasses();
+}
+
+bool UCompositeCorePluginSettings::IsAllowedPrimitiveClass(const UPrimitiveComponent* InPrimitiveComponent) const
+{
+	for (const UClass* ObjectClass : CachedDisabledPrimitiveClasses)
+	{
+		if (InPrimitiveComponent->IsA(ObjectClass))
+		{
+			return false;
+		}
+	}
+
+	return true;
+}
+
+void UCompositeCorePluginSettings::CacheDisabledClasses()
+{
+	CachedDisabledPrimitiveClasses.Reset(0);
+
+	for (const FSoftClassPath& ObjectClass : DisabledPrimitiveClasses)
+	{
+		if (const UClass* DisabledClass = ObjectClass.TryLoadClass<UObject>())
+		{
+			CachedDisabledPrimitiveClasses.Add(DisabledClass);
+		}
+	}
+}
diff --git a/Engine/Plugins/Experimental/Compositing/CompositeCore/Source/CompositeCore/Private/CompositeCoreSubsystem.cpp b/Engine/Plugins/Experimental/Compositing/CompositeCore/Source/CompositeCore/Private/CompositeCoreSubsystem.cpp
new file mode 100644
--- /dev/null
+++ b/Engine/Plugins/Experimental/Compositing/CompositeCore/Source/CompositeCore/Private/CompositeCoreSubsystem.cpp
@@ -0,0 +1,315 @@
+// Copyright Epic Games, Inc. All Rights Reserved.
+
+#include "CompositeCoreSubsystem.h"
+
+#include "CompositeCoreModule.h"
+#include "CompositeCoreSceneViewExtension.h"
+
+#include "Engine/RendererSettings.h"
+
+#if WITH_EDITOR
+#include "Framework/Notifications/NotificationManager.h"
+#include "HAL/PlatformFileManager.h"
+#include "ISettingsEditorModule.h"
+#include "Misc/ConfigCacheIni.h"
+#include "Modules/ModuleManager.h"
+#include "Widgets/Notifications/SNotificationList.h"
+#endif
+
+#define LOCTEXT_NAMESPACE "CompositeCoreSubsystem"
+
+namespace
+{
+#if WITH_EDITOR
+	void UpdateDependentPropertyInConfigFile(URendererSettings* RendererSettings, FProperty* RendererProperty)
+	{
+		FString RelativePath = RendererSettings->GetDefaultConfigFilename();
+		FString FullPath = FPaths::ConvertRelativePathToFull(RelativePath);
+
+		const bool bIsWriteable = !FPlatformFileManager::Get().GetPlatformFile().IsReadOnly(*FullPath);
+
+		if (!bIsWriteable)
+		{
+			FPlatformFileManager::Get().GetPlatformFile().SetReadOnly(*FullPath, false);
+		}
+
+		RendererSettings->UpdateSinglePropertyInConfigFile(RendererProperty, RendererSettings->GetDefaultConfigFilename());
+
+		// Restore original state for source control
+		if (!bIsWriteable)
+		{
+			FPlatformFileManager::Get().GetPlatformFile().SetReadOnly(*FullPath, true);
+		}
+	}
+#endif
+}
+
+UCompositeCoreSubsystem::UCompositeCoreSubsystem()
+{
+}
+
+void UCompositeCoreSubsystem::Initialize(FSubsystemCollectionBase& Collection)
+{
+	Super::Initialize(Collection);
+
+	UWorld* World = GetWorld();
+
+	if (IsValid(World))
+	{
+		CompositeCoreViewExtension = FSceneViewExtensions::NewExtension<FCompositeCoreSceneViewExtension>(World);
+	}
+}
+
+void UCompositeCoreSubsystem::Deinitialize()
+{
+	CompositeCoreViewExtension.Reset();
+	
+	Super::Deinitialize();
+}
+
+void UCompositeCoreSubsystem::Tick(float DeltaTime)
+{
+	Super::Tick(DeltaTime);
+}
+
+void UCompositeCoreSubsystem::RegisterPrimitive(UPrimitiveComponent* InPrimitiveComponent)
+{
+	TArray<UPrimitiveComponent*> PrimitiveComponents = { InPrimitiveComponent };
+
+	RegisterPrimitives(PrimitiveComponents);
+}
+
+void UCompositeCoreSubsystem::RegisterPrimitives(const TArray<UPrimitiveComponent*>& InPrimitiveComponents)
+{
+	const bool bValidSettings = ValidateProjectSettings();
+
+	if (bValidSettings && CompositeCoreViewExtension.IsValid())
+	{
+		CompositeCoreViewExtension->RegisterPrimitives_GameThread(InPrimitiveComponents);
+	}
+}
+
+void UCompositeCoreSubsystem::UnregisterPrimitive(UPrimitiveComponent* InPrimitiveComponent)
+{
+	TArray<UPrimitiveComponent*> PrimitiveComponents = { InPrimitiveComponent };
+
+	UnregisterPrimitives(PrimitiveComponents);
+}
+
+void UCompositeCoreSubsystem::UnregisterPrimitives(const TArray<UPrimitiveComponent*>& InPrimitiveComponents)
+{
+	if (CompositeCoreViewExtension.IsValid())
+	{
+		CompositeCoreViewExtension->UnregisterPrimitives_GameThread(InPrimitiveComponents);
+	}
+}
+
+void UCompositeCoreSubsystem::SetPostRenderWork(UE::CompositeCore::FPostRenderWork&& InWork)
+{
+	if (CompositeCoreViewExtension.IsValid())
+	{
+		CompositeCoreViewExtension->SetPostRenderWork_GameThread(MoveTemp(InWork));
+	}
+}
+
+void UCompositeCoreSubsystem::ResetPostRenderWork()
+{
+	if (CompositeCoreViewExtension.IsValid())
+	{
+		CompositeCoreViewExtension->ResetPostRenderWork_GameThread();
+	}
+}
+
+void UCompositeCoreSubsystem::SetBuiltInRenderPassOptions(const UE::CompositeCore::FBuiltInRenderPassOptions& InOptions)
+{
+	if (CompositeCoreViewExtension.IsValid())
+	{
+		CompositeCoreViewExtension->SetBuiltInRenderPassOptions_GameThread(InOptions);
+	}
+}
+
+void UCompositeCoreSubsystem::ResetBuiltInRenderPassOptions()
+{
+	if (CompositeCoreViewExtension.IsValid())
+	{
+		CompositeCoreViewExtension->SetBuiltInRenderPassOptions_GameThread(UE::CompositeCore::FBuiltInRenderPassOptions{});
+	}
+}
+
+void UCompositeCoreSubsystem::SetCompositeTextureOverride(UTexture* InTextureOverride)
+{
+	if (CompositeCoreViewExtension.IsValid())
+	{
+		CompositeCoreViewExtension->SetCompositeTextureOverride_GameThread(InTextureOverride);
+	}
+}
+
+bool UCompositeCoreSubsystem::IsProjectSettingsValid()
+{
+	const URendererSettings* RendererSettings = GetDefault<URendererSettings>();
+	check(RendererSettings);
+
+	const bool bValidSettings = 
+		RendererSettings->bEnableAlphaChannelInPostProcessing
+		&& RendererSettings->bDeferredSupportPrimitiveAlphaHoldout;
+
+	return bValidSettings;
+}
+
+bool UCompositeCoreSubsystem::ValidateProjectSettings()
+{
+	const bool bValidSettings = IsProjectSettingsValid();
+
+	if (!bValidSettings )
+	{
+#if WITH_EDITOR
+		URendererSettings* RendererSettings = GetMutableDefault<URendererSettings>();
+		check(RendererSettings);
+
+		// We inform the user and offer them the option to activate project settings.
+		UE_CALL_ONCE([&]{ PrimitiveHoldoutSettingsNotification(RendererSettings); });
+#else
+		UE_CALL_ONCE([&]{ UE_LOG(LogCompositeCore, Warning, TEXT("Both \"Alpha Output\" and \"Support Primitive Alpha Holdout\" project settings must be enabled for holdout composite.")); });
+#endif
+	}
+
+	return bValidSettings;
+}
+
+#if WITH_EDITOR
+void UCompositeCoreSubsystem::PrimitiveHoldoutSettingsNotification(URendererSettings* RendererSettings)
+{
+	const bool bAlphaOutputMissing = !RendererSettings->bEnableAlphaChannelInPostProcessing;
+	const bool bPrimitiveHoldoutMissing = !RendererSettings->bDeferredSupportPrimitiveAlphaHoldout;
+
+	const FText AlphaOutputSettingOption = LOCTEXT("HoldoutSetting_AlphaOutput", "\n- Alpha Output");
+	const FText PrimitiveHoldoutSettingOption = LOCTEXT("HoldoutSetting_PrimitiveHoldout", "\n- Support Primitive Alpha Holdout");
+	const FText HoldoutText = FText::Format(LOCTEXT("HoldoutSettingPrompt", "The following project setting(s) must be enabled for holdout composite:{0}{1}"), bAlphaOutputMissing ? AlphaOutputSettingOption : FText::GetEmpty(), bPrimitiveHoldoutMissing ? PrimitiveHoldoutSettingOption : FText::GetEmpty());
+	const FText HoldoutConfirmText = LOCTEXT("HoldoutSettingConfirm", "Enable");
+	const FText HoldoutCancelText = LOCTEXT("HoldoutSettingCancel", "Not Now");
+
+	/** Utility functions for notifications */
+	struct FSuppressDialogOptions
+	{
+		static bool ShouldSuppressModal()
+		{
+			bool bSuppressNotification = false;
+			GConfig->GetBool(TEXT("CompositeCore"), TEXT("SuppressCompositeCorePromptNotification"), bSuppressNotification, GEditorPerProjectIni);
+			return bSuppressNotification;
+		}
+
+		static ECheckBoxState GetDontAskAgainCheckBoxState()
+		{
+			return ShouldSuppressModal() ? ECheckBoxState::Checked : ECheckBoxState::Unchecked;
+		}
+
+		static void OnDontAskAgainCheckBoxStateChanged(ECheckBoxState NewState)
+		{
+			// If the user selects to not show this again, set that in the config so we know about it in between sessions
+			const bool bSuppressNotification = (NewState == ECheckBoxState::Checked);
+			GConfig->SetBool(TEXT("CompositeCore"), TEXT("SuppressCompositeCorePromptNotification"), bSuppressNotification, GEditorPerProjectIni);
+		}
+	};
+
+	// If the user has specified to supress this pop up, then just early out and exit	
+	if (FSuppressDialogOptions::ShouldSuppressModal())
+	{
+		return;
+	}
+
+	FSimpleDelegate OnConfirmDelegate = FSimpleDelegate::CreateLambda(
+		[WeakThis = MakeWeakObjectPtr(this), RendererSettings]()
+		{
+			if (IsValid(RendererSettings))
+			{
+				if (!RendererSettings->bDeferredSupportPrimitiveAlphaHoldout)
+				{
+					FProperty* Property = RendererSettings->GetClass()->FindPropertyByName(GET_MEMBER_NAME_CHECKED(URendererSettings, bDeferredSupportPrimitiveAlphaHoldout));
+					RendererSettings->PreEditChange(Property);
+
+					RendererSettings->bDeferredSupportPrimitiveAlphaHoldout = true;
+
+					FPropertyChangedEvent PropertyChangedEvent(Property, EPropertyChangeType::ValueSet, { RendererSettings });
+					RendererSettings->PostEditChangeProperty(PropertyChangedEvent);
+					UpdateDependentPropertyInConfigFile(RendererSettings, Property);
+
+					// SupportPrimitiveAlphaHoldout requires shader recompilation, ask for a restart.
+					FModuleManager::GetModuleChecked<ISettingsEditorModule>("SettingsEditor").OnApplicationRestartRequired();
+				}
+
+				if (!RendererSettings->bEnableAlphaChannelInPostProcessing)
+				{
+					FProperty* Property = RendererSettings->GetClass()->FindPropertyByName(GET_MEMBER_NAME_CHECKED(URendererSettings, bEnableAlphaChannelInPostProcessing));
+					RendererSettings->PreEditChange(Property);
+
+					RendererSettings->bEnableAlphaChannelInPostProcessing = true;
+
+					FPropertyChangedEvent PropertyChangedEvent(Property, EPropertyChangeType::ValueSet, { RendererSettings });
+					RendererSettings->PostEditChangeProperty(PropertyChangedEvent);
+					UpdateDependentPropertyInConfigFile(RendererSettings, Property);
+				}
+			}
+
+			TStrongObjectPtr<UCompositeCoreSubsystem> Subsystem = WeakThis.Pin();
+			if (Subsystem.IsValid())
+			{
+				TSharedPtr<SNotificationItem> NotificationItem = Subsystem->HoldoutNotificationItem.Pin();
+				if (NotificationItem.IsValid())
+				{
+					NotificationItem->SetCompletionState(SNotificationItem::CS_Success);
+					NotificationItem->ExpireAndFadeout();
+				}
+
+				Subsystem->HoldoutNotificationItem.Reset();
+			}
+		}
+	);
+
+	FSimpleDelegate OnCancelDelegate = FSimpleDelegate::CreateLambda(
+		[WeakThis = MakeWeakObjectPtr(this)]()
+		{
+			TStrongObjectPtr<UCompositeCoreSubsystem> Subsystem = WeakThis.Pin();
+			if (Subsystem.IsValid())
+			{
+				TSharedPtr<SNotificationItem> NotificationItem = Subsystem->HoldoutNotificationItem.Pin();
+				if (NotificationItem.IsValid())
+				{
+					NotificationItem->SetCompletionState(SNotificationItem::CS_None);
+					NotificationItem->ExpireAndFadeout();
+				}
+
+				Subsystem->HoldoutNotificationItem.Reset();
+			}
+		}
+	);
+
+	FNotificationInfo Info(HoldoutText);
+	Info.bFireAndForget = false;
+	Info.bUseLargeFont = false;
+	Info.bUseThrobber = false;
+	Info.bUseSuccessFailIcons = false;
+	Info.ButtonDetails.Add(FNotificationButtonInfo(HoldoutConfirmText, FText(), OnConfirmDelegate));
+	Info.ButtonDetails.Add(FNotificationButtonInfo(HoldoutCancelText, FText(), OnCancelDelegate));
+
+	// Add a "Don't show this again" option
+	Info.CheckBoxState = TAttribute<ECheckBoxState>::Create(&FSuppressDialogOptions::GetDontAskAgainCheckBoxState);
+	Info.CheckBoxStateChanged = FOnCheckStateChanged::CreateStatic(&FSuppressDialogOptions::OnDontAskAgainCheckBoxStateChanged);
+	Info.CheckBoxText = LOCTEXT("DontShowThisAgainCheckBoxMessage", "Don't show this again");
+
+	if (HoldoutNotificationItem.IsValid())
+	{
+		HoldoutNotificationItem.Pin()->ExpireAndFadeout();
+		HoldoutNotificationItem.Reset();
+	}
+
+	HoldoutNotificationItem = FSlateNotificationManager::Get().AddNotification(Info);
+
+	if (HoldoutNotificationItem.IsValid())
+	{
+		HoldoutNotificationItem.Pin()->SetCompletionState(SNotificationItem::CS_Pending);
+	}
+}
+#endif
+
+#undef LOCTEXT_NAMESPACE
+
diff --git a/Engine/Plugins/Experimental/Compositing/CompositeCore/Source/CompositeCore/Private/HoldoutCompositeComponent.cpp b/Engine/Plugins/Experimental/Compositing/CompositeCore/Source/CompositeCore/Private/HoldoutCompositeComponent.cpp
new file mode 100644
--- /dev/null
+++ b/Engine/Plugins/Experimental/Compositing/CompositeCore/Source/CompositeCore/Private/HoldoutCompositeComponent.cpp
@@ -0,0 +1,192 @@
+// Copyright Epic Games, Inc. All Rights Reserved.
+
+#include "HoldoutCompositeComponent.h"
+
+#include "Components/PrimitiveComponent.h"
+#include "Engine/Engine.h"
+#include "CompositeCoreModule.h"
+#include "CompositeCoreSettings.h"
+#include "CompositeCoreSubsystem.h"
+
+#if WITH_EDITOR
+#include "Framework/Notifications/NotificationManager.h"
+#include "Widgets/Notifications/SNotificationList.h"
+#endif
+
+#define LOCTEXT_NAMESPACE "CompositeCore"
+
+namespace UE
+{
+	namespace CompositeCore
+	{
+		namespace Private
+		{
+			// Find list of primitives from the parent component and its children. We need to traverse children to support objects such as UText3DComponent.
+			static TArray<UPrimitiveComponent*> FindPrimitiveComponents(USceneComponent* InParentComponent)
+			{
+				const UCompositeCorePluginSettings* PluginSettings = GetDefault<UCompositeCorePluginSettings>();
+				check(PluginSettings);
+
+				TArray<UPrimitiveComponent*> OutPrimitiveComponents;
+
+				if (IsValid(InParentComponent))
+				{
+					if (UPrimitiveComponent* PrimitiveComponent = Cast<UPrimitiveComponent>(InParentComponent))
+					{
+						if (PluginSettings->IsAllowedPrimitiveClass(PrimitiveComponent))
+						{
+							OutPrimitiveComponents.Add(PrimitiveComponent);
+						}
+					}
+
+					TArray<USceneComponent*> ParentChildComponents;
+					InParentComponent->GetChildrenComponents(true /*bIncludeAllDescendants*/, ParentChildComponents);
+
+					for (USceneComponent* ParentChild : ParentChildComponents)
+					{
+						if (UPrimitiveComponent* PrimitiveComponent = Cast<UPrimitiveComponent>(ParentChild))
+						{
+							if (PluginSettings->IsAllowedPrimitiveClass(PrimitiveComponent))
+							{
+								OutPrimitiveComponents.Add(PrimitiveComponent);
+							}
+						}
+					}
+				}
+
+				return OutPrimitiveComponents;
+			}
+		}
+	}
+}
+
+UHoldoutCompositeComponent::UHoldoutCompositeComponent(const FObjectInitializer& ObjectInitializer)
+	: Super(ObjectInitializer)
+{
+	if (UE::CompositeCore::IsRegisterPrimitivesOnTickEnabled())
+	{
+		PrimaryComponentTick.TickGroup = TG_PostUpdateWork;
+		PrimaryComponentTick.bStartWithTickEnabled = true;
+		PrimaryComponentTick.bCanEverTick = true;
+		PrimaryComponentTick.bTickEvenWhenPaused = true;
+		bTickInEditor = true;
+		bAutoActivate = true;
+	}
+}
+
+void UHoldoutCompositeComponent::OnRegister()
+{
+	Super::OnRegister();
+
+	RegisterCompositeImpl();
+}
+
+void UHoldoutCompositeComponent::OnUnregister()
+{
+	UnregisterCompositeImpl();
+
+	Super::OnUnregister();
+}
+
+void UHoldoutCompositeComponent::TickComponent(float DeltaTime, ELevelTick TickType, FActorComponentTickFunction* ThisTickFunction)
+{
+	Super::TickComponent(DeltaTime, TickType, ThisTickFunction);
+
+	constexpr bool bIsTick = true;
+	RegisterCompositeImpl(bIsTick);
+}
+
+void UHoldoutCompositeComponent::DetachFromComponent(const FDetachmentTransformRules& DetachmentRules)
+{
+	// Note: We also unregister here while the attached parent pointer is still valid.
+	UnregisterCompositeImpl();
+
+	Super::DetachFromComponent(DetachmentRules);
+}
+
+void UHoldoutCompositeComponent::OnAttachmentChanged()
+{
+	Super::OnAttachmentChanged();
+
+	UnregisterCompositeImpl();
+
+	USceneComponent* SceneComponent = GetAttachParent();
+	if (IsValid(SceneComponent))
+	{
+		TArray<UPrimitiveComponent*> ParentPrimitives = UE::CompositeCore::Private::FindPrimitiveComponents(SceneComponent);
+		if (ParentPrimitives.IsEmpty())
+		{
+#if WITH_EDITOR
+			FNotificationInfo Info(LOCTEXT("CompositeParentNotification",
+				"The composite component must be parented to a primitive component (or one that has primitives)."));
+			Info.ExpireDuration = 5.0f;
+
+			FSlateNotificationManager::Get().AddNotification(Info);
+#endif
+		}
+		else
+		{
+			RegisterCompositeImpl();
+		}
+	}
+}
+
+bool UHoldoutCompositeComponent::IsEnabled() const
+{
+	return bIsEnabled;
+}
+
+void UHoldoutCompositeComponent::SetEnabled(bool bInIsEnabled)
+{
+	if (bIsEnabled != bInIsEnabled)
+	{
+		bIsEnabled = bInIsEnabled;
+
+		if (bIsEnabled)
+		{
+			RegisterCompositeImpl();
+		}
+		else
+		{
+			UnregisterCompositeImpl();
+		}
+	}
+}
+
+void UHoldoutCompositeComponent::RegisterCompositeImpl(bool bIsTick)
+{
+	using namespace UE::CompositeCore;
+
+	if (IsRegisterPrimitivesOnTickEnabled() != bIsTick)
+	{
+		return;
+	}
+
+	if (!bIsEnabled)
+	{
+		return;
+	}
+
+	UCompositeCoreSubsystem* Subsystem = UWorld::GetSubsystem<UCompositeCoreSubsystem>(GetWorld());
+	TArray<UPrimitiveComponent*> ParentPrimitives = Private::FindPrimitiveComponents(GetAttachParent());
+
+	if (IsValid(Subsystem) && !ParentPrimitives.IsEmpty())
+	{
+		Subsystem->RegisterPrimitives(ParentPrimitives);
+	}
+}
+
+void UHoldoutCompositeComponent::UnregisterCompositeImpl()
+{
+	using namespace UE::CompositeCore::Private;
+
+	UCompositeCoreSubsystem* Subsystem = UWorld::GetSubsystem<UCompositeCoreSubsystem>(GetWorld());
+	TArray<UPrimitiveComponent*> ParentPrimitives = FindPrimitiveComponents(GetAttachParent());
+
+	if (IsValid(Subsystem) && !ParentPrimitives.IsEmpty())
+	{
+		Subsystem->UnregisterPrimitives(ParentPrimitives);
+	}
+}
+
+#undef LOCTEXT_NAMESPACE
diff --git a/Engine/Plugins/Experimental/Compositing/CompositeCore/Source/CompositeCore/Private/Passes/CompositeCorePassMergeProxy.cpp b/Engine/Plugins/Experimental/Compositing/CompositeCore/Source/CompositeCore/Private/Passes/CompositeCorePassMergeProxy.cpp
new file mode 100644
--- /dev/null
+++ b/Engine/Plugins/Experimental/Compositing/CompositeCore/Source/CompositeCore/Private/Passes/CompositeCorePassMergeProxy.cpp
@@ -0,0 +1,212 @@
+// Copyright Epic Games, Inc. All Rights Reserved.
+
+#include "Passes/CompositeCorePassMergeProxy.h"
+
+#include "CompositeCoreSettings.h"
+
+#include "PixelShaderUtils.h"
+#include "PostProcess/LensDistortion.h"
+#include "PostProcess/PostProcessMaterialInputs.h"
+#include "RenderGraphBuilder.h"
+#include "RenderGraphUtils.h"
+#include "SceneView.h"
+#include "SystemTextures.h"
+
+DECLARE_GPU_STAT_NAMED(FCompositeCoreMerge, TEXT("CompositeCore.Merge"));
+
+class FCompositeCoreMergeShader : public FGlobalShader
+{
+	DECLARE_GLOBAL_SHADER(FCompositeCoreMergeShader);
+	SHADER_USE_PARAMETER_STRUCT(FCompositeCoreMergeShader, FGlobalShader);
+
+	class FUseDistortion : SHADER_PERMUTATION_BOOL("USE_DISTORTION");
+	using FPermutationDomain = TShaderPermutationDomain<FUseDistortion>;
+
+	BEGIN_SHADER_PARAMETER_STRUCT(FInputPassTextureParameters, )
+		SHADER_PARAMETER_RDG_TEXTURE(Texture2D, Texture)
+		SHADER_PARAMETER_SAMPLER(SamplerState, Sampler)
+		SHADER_PARAMETER(uint32, bAlphaIsOpacity)
+		SHADER_PARAMETER(uint32, SourceEncoding)
+		SHADER_PARAMETER(uint32, DistortionUV)
+	END_SHADER_PARAMETER_STRUCT()
+
+	BEGIN_SHADER_PARAMETER_STRUCT(FParameters, )
+		SHADER_PARAMETER_STRUCT_REF(FViewUniformShaderParameters, View)
+		SHADER_PARAMETER_STRUCT(FScreenPassTextureViewportParameters, Input0)
+		SHADER_PARAMETER_STRUCT(FScreenPassTextureViewportParameters, Input1)
+		SHADER_PARAMETER_STRUCT(FScreenPassTextureViewportParameters, Output)
+		SHADER_PARAMETER_STRUCT(FInputPassTextureParameters, Tex0)
+		SHADER_PARAMETER_STRUCT(FInputPassTextureParameters, Tex1)
+		SHADER_PARAMETER_RDG_TEXTURE(Texture2D, DistortingDisplacementTexture)
+		SHADER_PARAMETER_SAMPLER(SamplerState, DistortingDisplacementSampler)
+		SHADER_PARAMETER_RDG_TEXTURE(Texture2D, UndistortingDisplacementTexture)
+		SHADER_PARAMETER_SAMPLER(SamplerState, UndistortingDisplacementSampler)
+		SHADER_PARAMETER(FVector2f, DisplayGamma)
+		SHADER_PARAMETER(uint32, MergeOp)
+		SHADER_PARAMETER(uint32, OutputEncoding)
+		SHADER_PARAMETER(uint32, bFlags)
+		RENDER_TARGET_BINDING_SLOTS()
+	END_SHADER_PARAMETER_STRUCT()
+};
+IMPLEMENT_GLOBAL_SHADER(FCompositeCoreMergeShader, "/Plugin/CompositeCore/Private/CompositeCoreMerge.usf", "MainPS", SF_Pixel);
+
+namespace UE
+{
+	namespace CompositeCore
+	{
+		enum class ESourceEncoding : uint32
+		{
+			Linear = 0,
+			Gamma = 1,
+			sRGB = 2,
+		};
+
+		enum class EDistortionUV
+		{
+			None = 0,
+			Distorted = 1,
+			Undistorted = 2
+		};
+
+		FCompositeCoreMergeShader::FInputPassTextureParameters GetPassTextureParameters(
+			const FScreenPassTexture& InTexture,
+			const ECompositeCorePassInput InputType,
+			const FSceneViewFamily* Family,
+			const ISceneViewExtension::EPostProcessingPass InLocation,
+			const bool bUseLensDistortion)
+		{
+			FCompositeCoreMergeShader::FInputPassTextureParameters Parameters;
+
+			Parameters.Texture = InTexture.Texture;
+			Parameters.Sampler = TStaticSamplerState<SF_Bilinear>::GetRHI();
+			Parameters.bAlphaIsOpacity = (InputType == ECompositeCorePassInput::CompositeTexture) ? 1 : 0;
+			Parameters.SourceEncoding = static_cast<uint32>(ESourceEncoding::Linear);
+
+			// After-tonemap scene color may have encoding manually applied, as opposed to _SRGB textures
+			if (InputType == ECompositeCorePassInput::Default && InLocation >= ISceneViewExtension::EPostProcessingPass::Tonemap)
+			{
+				if ((Family->EngineShowFlags.Tonemapper == 0) || (Family->EngineShowFlags.PostProcessing == 0))
+				{
+					Parameters.SourceEncoding = static_cast<uint32>(ESourceEncoding::Gamma);
+				}
+				else if (Family->SceneCaptureSource == SCS_FinalColorLDR)
+				{
+					Parameters.SourceEncoding = static_cast<uint32>(ESourceEncoding::sRGB);
+				}
+			}
+
+			EDistortionUV DistortionUV = EDistortionUV::None;
+
+			if (bUseLensDistortion)
+			{
+				const bool bOutputIsDistorted = !(InLocation == ISceneViewExtension::EPostProcessingPass::SSRInput ||
+					InLocation == ISceneViewExtension::EPostProcessingPass::TranslucencyAfterDOF);
+
+				if (bOutputIsDistorted)
+				{
+					// If output is distorted, we must sample undistorted textures with undistorted UVs
+					if (InputType == ECompositeCorePassInput::CompositeRenderPass)
+					{
+						DistortionUV = EDistortionUV::Undistorted;
+					}
+				}
+				else
+				{
+					// If output is undistorted, we must sample distorted textures with distorted UVs
+					if (InputType == ECompositeCorePassInput::Default || InputType == ECompositeCorePassInput::CompositeTexture)
+					{
+						DistortionUV = EDistortionUV::Distorted;
+					}
+				}
+			}
+
+			Parameters.DistortionUV = static_cast<uint32>(DistortionUV);
+
+			return Parameters;
+		}
+
+		FMergePassProxy::FMergePassProxy(ECompositeCorePassInput Input0, ECompositeCorePassInput Input1)
+			: ICompositeCorePassProxy(Input0)
+		{
+			InputTypes.Add(Input1);
+		}
+
+		FScreenPassTexture FMergePassProxy::Add(FRDGBuilder& GraphBuilder, const FSceneView& InView, const FPostProcessMaterialInputs& Inputs, ISceneViewExtension::EPostProcessingPass InLocation) const
+		{
+			RDG_EVENT_SCOPE_STAT(GraphBuilder, FCompositeCoreMerge, "CompositeCore.Merge");
+			RDG_GPU_STAT_SCOPE(GraphBuilder, FCompositeCoreMerge);
+
+			FScreenPassTexture Input0 = FScreenPassTexture::CopyFromSlice(GraphBuilder, Inputs.GetInput(static_cast<EPostProcessMaterialInput>(0)));
+			FScreenPassTexture Input1 = FScreenPassTexture::CopyFromSlice(GraphBuilder, Inputs.GetInput(static_cast<EPostProcessMaterialInput>(1)));
+
+			FGlobalShaderMap* GlobalShaderMap = GetGlobalShaderMap(InView.GetFeatureLevel());
+			const FSceneViewFamily* Family = InView.Family;
+
+			ESourceEncoding OutputEncoding = ESourceEncoding::Linear;
+
+			// When running after tonemap, we assume that input0 (scene color) has an encoding.
+			if (InLocation >= ISceneViewExtension::EPostProcessingPass::Tonemap)
+			{
+				if ((Family->EngineShowFlags.Tonemapper == 0) || (Family->EngineShowFlags.PostProcessing == 0))
+				{
+					OutputEncoding = ESourceEncoding::Gamma;
+				}
+				else if (Family->SceneCaptureSource == SCS_FinalColorLDR)
+				{
+					OutputEncoding = ESourceEncoding::sRGB;
+				}
+			}
+
+			FScreenPassRenderTarget Output = Inputs.OverrideOutput;
+			if (!Output.IsValid())
+			{
+				Output = FScreenPassRenderTarget::CreateFromInput(GraphBuilder, Input0, InView.GetOverwriteLoadAction(), TEXT("CompositeCoreMergeOutput"));
+			}
+
+			const FLensDistortionLUT& LensDistortionLUT = LensDistortion::GetLUTUnsafe(InView);
+			const bool bUseLensDistortion = LensDistortionLUT.IsEnabled() && (LensDistortion::GetPassLocationUnsafe(InView) == LensDistortion::EPassLocation::TSR);
+
+			FCompositeCoreMergeShader::FParameters* PassParameters = GraphBuilder.AllocParameters<FCompositeCoreMergeShader::FParameters>();
+			PassParameters->View = InView.ViewUniformBuffer;
+			PassParameters->Input0 = GetScreenPassTextureViewportParameters(FScreenPassTextureViewport(Input0));
+			PassParameters->Input1 = GetScreenPassTextureViewportParameters(FScreenPassTextureViewport(Input1));
+			PassParameters->Output = GetScreenPassTextureViewportParameters(FScreenPassTextureViewport(Output));
+			PassParameters->Tex0 = GetPassTextureParameters(Input0, InputTypes[0], Family, InLocation, bUseLensDistortion);
+			PassParameters->Tex1 = GetPassTextureParameters(Input1, InputTypes[1], Family, InLocation, bUseLensDistortion);
+			if (bUseLensDistortion)
+			{
+				PassParameters->DistortingDisplacementTexture = LensDistortionLUT.DistortingDisplacementTexture;
+				PassParameters->UndistortingDisplacementTexture = LensDistortionLUT.UndistortingDisplacementTexture;
+			}
+			else
+			{
+				PassParameters->UndistortingDisplacementTexture = GSystemTextures.GetBlackDummy(GraphBuilder);
+				PassParameters->DistortingDisplacementTexture = GSystemTextures.GetBlackDummy(GraphBuilder);
+			}
+			PassParameters->UndistortingDisplacementSampler = TStaticSamplerState<SF_Bilinear>::GetRHI();
+			PassParameters->DistortingDisplacementSampler = TStaticSamplerState<SF_Bilinear>::GetRHI();
+			PassParameters->DisplayGamma = FVector2f(Family->RenderTarget->GetDisplayGamma(), 1.0f / Family->RenderTarget->GetDisplayGamma());
+			PassParameters->MergeOp = static_cast<uint32>(MergeOp);
+			PassParameters->OutputEncoding = static_cast<uint32>(OutputEncoding);
+			PassParameters->bFlags = static_cast<uint32>(bUnpremulInput1 << 1) | static_cast<uint32>(bOutputSceneColor);
+			PassParameters->RenderTargets[0] = Output.GetRenderTargetBinding();
+
+			FCompositeCoreMergeShader::FPermutationDomain PermutationVector;
+			PermutationVector.Set<FCompositeCoreMergeShader::FUseDistortion>(bUseLensDistortion);
+			TShaderMapRef<FCompositeCoreMergeShader> PixelShader(GlobalShaderMap, PermutationVector);
+
+			FPixelShaderUtils::AddFullscreenPass(
+				GraphBuilder,
+				GlobalShaderMap,
+				RDG_EVENT_NAME("CompositeCore.Merge (%dx%d) PS",
+					Output.ViewRect.Width(), Output.ViewRect.Height()),
+				PixelShader,
+				PassParameters,
+				Output.ViewRect
+			);
+
+			return MoveTemp(Output);
+		}
+	}
+}
+
diff --git a/Engine/Plugins/Experimental/Compositing/CompositeCore/Source/CompositeCore/Private/Passes/CompositeCorePassProxy.cpp b/Engine/Plugins/Experimental/Compositing/CompositeCore/Source/CompositeCore/Private/Passes/CompositeCorePassProxy.cpp
new file mode 100644
--- /dev/null
+++ b/Engine/Plugins/Experimental/Compositing/CompositeCore/Source/CompositeCore/Private/Passes/CompositeCorePassProxy.cpp
@@ -0,0 +1,46 @@
+// Copyright Epic Games, Inc. All Rights Reserved.
+
+#include "Passes/CompositeCorePassProxy.h"
+
+#include "Passes/CompositeCorePassMergeProxy.h"
+
+namespace UE
+{
+	namespace CompositeCore
+	{
+		FBuiltInRenderPassOptions::FBuiltInRenderPassOptions()
+			: ViewUserFlagsOverride{}
+			, bEnableUnlitViewmode{true}
+			, OutputDilationSize{1}
+		{
+		}
+
+		const FPostRenderWork& FPostRenderWork::GetDefault()
+		{
+			static FPostRenderWork DefaultRenderWork = []() -> FPostRenderWork
+				{
+					FPostRenderWork RenderWork;
+
+					constexpr ECompositeCorePassInput Input0 = ECompositeCorePassInput::Default;
+					constexpr ECompositeCorePassInput Input1 = ECompositeCorePassInput::CompositeRenderPass;
+
+					TArray<const ICompositeCorePassProxy*>& SSRInputWork = RenderWork.FramePasses.Add(ISceneViewExtension::EPostProcessingPass::SSRInput);
+					TArray<const ICompositeCorePassProxy*>& AfterTonemapWork = RenderWork.FramePasses.Add(ISceneViewExtension::EPostProcessingPass::Tonemap);
+
+					SSRInputWork.Add(RenderWork.FrameAllocator->Create<FMergePassProxy>(Input0, Input1));
+					AfterTonemapWork.Add(RenderWork.FrameAllocator->Create<FMergePassProxy>(Input0, Input1));
+
+					return RenderWork;
+				}();
+
+			return DefaultRenderWork;
+		}
+
+		FPostRenderWork::FPostRenderWork()
+			: InputProcessing{}
+			, FramePasses{}
+			, FrameAllocator{new FSceneRenderingBulkObjectAllocator}
+		{ }
+	}
+}
+
diff --git a/Engine/Plugins/Experimental/Compositing/CompositeCore/Source/CompositeCore/Public/CompositeCoreModule.h b/Engine/Plugins/Experimental/Compositing/CompositeCore/Source/CompositeCore/Public/CompositeCoreModule.h
new file mode 100644
--- /dev/null
+++ b/Engine/Plugins/Experimental/Compositing/CompositeCore/Source/CompositeCore/Public/CompositeCoreModule.h
@@ -0,0 +1,21 @@
+// Copyright Epic Games, Inc. All Rights Reserved.
+
+#pragma once
+
+#include "Modules/ModuleInterface.h"
+#include "Logging/LogMacros.h"
+#include "Stats/Stats.h"
+
+COMPOSITECORE_API DECLARE_LOG_CATEGORY_EXTERN(LogCompositeCore, Log, All);
+
+DECLARE_STATS_GROUP(TEXT("CompositeCore"), STATGROUP_CompositeCore, STATCAT_Advanced);
+
+class FCompositeCoreModule : public IModuleInterface
+{
+public:
+	// IModuleInterface interface
+	virtual void StartupModule() override;
+	virtual void ShutdownModule() override;
+	// End of IModuleInterface interface
+};
+
diff --git a/Engine/Plugins/Experimental/Compositing/CompositeCore/Source/CompositeCore/Public/CompositeCoreSettings.h b/Engine/Plugins/Experimental/Compositing/CompositeCore/Source/CompositeCore/Public/CompositeCoreSettings.h
new file mode 100644
--- /dev/null
+++ b/Engine/Plugins/Experimental/Compositing/CompositeCore/Source/CompositeCore/Public/CompositeCoreSettings.h
@@ -0,0 +1,66 @@
+// Copyright Epic Games, Inc. All Rights Reserved.
+
+#pragma once
+
+#include "CoreTypes.h"
+#include "UObject/Object.h"
+#include "UObject/ObjectMacros.h"
+#include "Engine/DeveloperSettings.h"
+
+#include "CompositeCoreSettings.generated.h"
+
+class UPrimitiveComponent;
+
+namespace UE
+{
+	namespace CompositeCore
+	{
+		COMPOSITECORE_API bool IsRegisterPrimitivesOnTickEnabled();
+	}
+}
+
+/**
+ * Settings for the CompositeCore module.
+ */
+UCLASS(MinimalAPI, config = Engine, meta = (DisplayName = "Composite Core"), MinimalAPI)
+class UCompositeCorePluginSettings : public UDeveloperSettings
+{
+	GENERATED_BODY()
+
+public:
+	UCompositeCorePluginSettings();
+
+	//~ Begin UDeveloperSettings interface
+	virtual FName GetCategoryName() const;
+#if WITH_EDITOR
+	virtual FText GetSectionText() const override;
+	virtual FName GetSectionName() const override;
+	virtual void PostEditChangeProperty(FPropertyChangedEvent& PropertyChangedEvent) override;
+#endif
+	//~ End UDeveloperSettings interface
+
+	virtual void PostInitProperties() override;
+
+public:
+
+	// Ensure that the primitive class is allowed by checking against plugin settings.
+	COMPOSITECORE_API bool IsAllowedPrimitiveClass(const UPrimitiveComponent* InPrimitiveComponent) const;
+
+	/** Primitive component classes that do not support the composite pipeline.*/
+	UPROPERTY(config, EditAnywhere, Category = CompositeCore)
+	TArray<FSoftClassPath> DisabledPrimitiveClasses;
+
+	/** Composite (scene view extension) post-processing priority, which defaults to before OpenColorIO. */
+	UPROPERTY(config, EditAnywhere, Category = CompositeCore, AdvancedDisplay)
+	int32 SceneViewExtensionPriority;
+
+private:
+
+	/** Cache disable class pointers to prevent load object calls during allowed class checks. */
+	void CacheDisabledClasses();
+	
+	/** Cached list of disallowed class types. */
+	TArray<const UClass*> CachedDisabledPrimitiveClasses;
+};
+
+
diff --git a/Engine/Plugins/Experimental/Compositing/CompositeCore/Source/CompositeCore/Public/CompositeCoreSubsystem.h b/Engine/Plugins/Experimental/Compositing/CompositeCore/Source/CompositeCore/Public/CompositeCoreSubsystem.h
new file mode 100644
--- /dev/null
+++ b/Engine/Plugins/Experimental/Compositing/CompositeCore/Source/CompositeCore/Public/CompositeCoreSubsystem.h
@@ -0,0 +1,94 @@
+// Copyright Epic Games, Inc. All Rights Reserved.
+
+#pragma once
+
+#include "MaterialShared.h"
+#include "Subsystems/WorldSubsystem.h"
+
+#include "CompositeCoreSubsystem.generated.h"
+
+class FCompositeCoreSceneViewExtension;
+struct FCompositeCoreSettings;
+class SNotificationItem;
+class UPrimitiveComponent;
+class URendererSettings;
+
+namespace UE::CompositeCore
+{
+	struct FBuiltInRenderPassOptions;
+	struct FPostRenderWork;
+}
+
+/**
+ * Composite subsytem used as an interface to the (private) scene view extension.
+ */
+UCLASS(BlueprintType, Transient)
+class COMPOSITECORE_API UCompositeCoreSubsystem : public UTickableWorldSubsystem
+{
+	GENERATED_BODY()
+
+public:
+	UCompositeCoreSubsystem();
+
+	// USubsystem implementation Begin
+	virtual void Initialize(FSubsystemCollectionBase& Collection) override;
+	virtual void Deinitialize() override;
+	// USubsystem implementation End
+
+	// FTickableGameObject implementation Begin
+	virtual bool IsTickableInEditor() const override { return true; }
+	virtual void Tick(float DeltaTime) override;
+	// FTickableGameObject implementation End
+
+	virtual TStatId GetStatId() const override { RETURN_QUICK_DECLARE_CYCLE_STAT(UCompositeCoreSubsystem, STATGROUP_Tickables); }
+
+	/* Register a single primitive for compositing. */
+	UFUNCTION(BlueprintCallable, Category = "Composite Core")
+	void RegisterPrimitive(UPrimitiveComponent* InPrimitiveComponent);
+
+	/* Register multiple primitives for compositing. */
+	void RegisterPrimitives(const TArray<UPrimitiveComponent*>& InPrimitiveComponents);
+
+	/* Unregister a single primitive from compositing. */
+	UFUNCTION(BlueprintCallable, Category = "Composite Core")
+	void UnregisterPrimitive(UPrimitiveComponent* InPrimitiveComponent);
+
+	/* Unregister multiple primitives from compositing. */
+	void UnregisterPrimitives(const TArray<UPrimitiveComponent*>& InPrimitiveComponents);
+
+	/* Set post-processing frame render work. */
+	void SetPostRenderWork(UE::CompositeCore::FPostRenderWork&& InWork);
+
+	/* Reset post-processing frame render work. */
+	void ResetPostRenderWork();
+
+	/* Set built-in composite render pass options. */
+	void SetBuiltInRenderPassOptions(const UE::CompositeCore::FBuiltInRenderPassOptions& InOptions);
+
+	/* Reset built-in composite render pass options. */
+	void ResetBuiltInRenderPassOptions();
+
+	/* Set a texture override to be used instead of the separate composited render. */
+	UFUNCTION(BlueprintCallable, Category = "Composite Core")
+	virtual void SetCompositeTextureOverride(UTexture* InTextureOverride);
+
+	/** True if the project settings are valid for the CompositeCore plugin to work. */
+	static bool IsProjectSettingsValid();
+
+private:
+
+	/* Returns true if the (renderer) project settings are correctly enabled for the composite to be active. */
+	bool ValidateProjectSettings();
+
+#if WITH_EDITOR
+	/* Toast notification to ask users to enable the missing project settings. */
+	void PrimitiveHoldoutSettingsNotification(URendererSettings* RendererSettings);
+
+	/* Toast notification item. */
+	TWeakPtr<SNotificationItem> HoldoutNotificationItem;
+#endif
+
+	/* Owned scene view extension. */
+	TSharedPtr<FCompositeCoreSceneViewExtension, ESPMode::ThreadSafe> CompositeCoreViewExtension;
+};
+
diff --git a/Engine/Plugins/Experimental/Compositing/CompositeCore/Source/CompositeCore/Public/HoldoutCompositeComponent.h b/Engine/Plugins/Experimental/Compositing/CompositeCore/Source/CompositeCore/Public/HoldoutCompositeComponent.h
new file mode 100644
--- /dev/null
+++ b/Engine/Plugins/Experimental/Compositing/CompositeCore/Source/CompositeCore/Public/HoldoutCompositeComponent.h
@@ -0,0 +1,49 @@
+// Copyright Epic Games, Inc. All Rights Reserved.
+
+#pragma once
+
+#include "CoreTypes.h"
+#include "Components/SceneComponent.h"
+
+#include "HoldoutCompositeComponent.generated.h"
+
+UCLASS(MinimalAPI, ClassGroup = Rendering, HideCategories=(Activation, Transform, Lighting, Rendering, Tags, Cooking, Physics, LOD, AssetUserData, Navigation), editinlinenew, meta = (BlueprintSpawnableComponent))
+class UHoldoutCompositeComponent : public USceneComponent
+{
+	GENERATED_UCLASS_BODY()
+
+public:
+	//~ Begin UActorComponent interface
+	virtual void OnRegister() override;
+	virtual void OnUnregister() override;
+	virtual void TickComponent(float DeltaTime, enum ELevelTick TickType, FActorComponentTickFunction* ThisTickFunction) override;
+	//~ End UActorComponent interface
+	
+	//~ Begin USceneComponent Interface.
+	virtual void OnAttachmentChanged() override;
+	virtual void DetachFromComponent(const FDetachmentTransformRules& DetachmentRules) override;
+	//~ End USceneComponent Interface
+
+	/* Get the enabled state of the component. */
+	UFUNCTION(BlueprintGetter)
+	COMPOSITECORE_API bool IsEnabled() const;
+
+	/* Set the enabled state of the component. */
+	UFUNCTION(BlueprintSetter)
+	COMPOSITECORE_API void SetEnabled(bool bInEnabled);
+
+private:
+
+	/* Private implementation of the register method. */
+	void RegisterCompositeImpl(bool bIsTick = false);
+
+	/* Private implementation of the unregister method. */
+	void UnregisterCompositeImpl();
+
+private:
+
+	/* Whether or not the component activates the composite. */
+	UPROPERTY(EditAnywhere, BlueprintGetter = IsEnabled, BlueprintSetter = SetEnabled, Category = "CompositeCore", meta = (AllowPrivateAccess = true))
+	bool bIsEnabled = true;
+};
+
diff --git a/Engine/Plugins/Experimental/Compositing/CompositeCore/Source/CompositeCore/Public/Passes/CompositeCorePassMergeProxy.h b/Engine/Plugins/Experimental/Compositing/CompositeCore/Source/CompositeCore/Public/Passes/CompositeCorePassMergeProxy.h
new file mode 100644
--- /dev/null
+++ b/Engine/Plugins/Experimental/Compositing/CompositeCore/Source/CompositeCore/Public/Passes/CompositeCorePassMergeProxy.h
@@ -0,0 +1,41 @@
+// Copyright Epic Games, Inc. All Rights Reserved.
+
+#pragma once
+
+#include "CoreTypes.h"
+#include "CompositeCorePassProxy.h"
+
+#include "CompositeCorePassMergeProxy.generated.h"
+
+UENUM()
+enum class ECompositeCoreMergeOp : uint8
+{
+	Over,
+	Add,
+	Multiply,
+	Divide,
+};
+
+namespace UE
+{
+	namespace CompositeCore
+	{
+		class COMPOSITECORE_API FMergePassProxy : public ICompositeCorePassProxy
+		{
+		public:
+			/** Constructor. */
+			FMergePassProxy(ECompositeCorePassInput Input0, ECompositeCorePassInput Input1);
+
+			FScreenPassTexture Add(FRDGBuilder& GraphBuilder, const FSceneView& InView, const FPostProcessMaterialInputs& Inputs, ISceneViewExtension::EPostProcessingPass InLocation) const override;
+
+			/** Merge operation. */
+			ECompositeCoreMergeOp MergeOp = ECompositeCoreMergeOp::Over;
+
+			/* Unpremultiplication required for recovering opaque colors for the default over blend. */
+			bool bUnpremulInput1 = true;
+
+			/* Output is assumed to be scene color with inverted alpha. */
+			bool bOutputSceneColor = true;
+		};
+	}
+}
diff --git a/Engine/Plugins/Experimental/Compositing/CompositeCore/Source/CompositeCore/Public/Passes/CompositeCorePassProxy.h b/Engine/Plugins/Experimental/Compositing/CompositeCore/Source/CompositeCore/Public/Passes/CompositeCorePassProxy.h
new file mode 100644
--- /dev/null
+++ b/Engine/Plugins/Experimental/Compositing/CompositeCore/Source/CompositeCore/Public/Passes/CompositeCorePassProxy.h
@@ -0,0 +1,95 @@
+// Copyright Epic Games, Inc. All Rights Reserved.
+
+#pragma once
+
+#include "CoreTypes.h"
+#include "Containers/ContainersFwd.h"
+#include "SceneViewExtension.h"
+#include "UObject/Object.h"
+
+#include "CompositeCorePassProxy.generated.h"
+
+class FRDGBuilder;
+class FSceneView;
+struct FScreenPassTexture;
+struct FPostProcessMaterialInputs;
+class ICompositeCorePassProxy;
+
+/* Composite texture input type. */
+UENUM()
+enum class ECompositeCorePassInput : uint8
+{
+	/** Default input (usually scene color), or result of the previous pass. */
+	Default,
+	/** Texture from the built-in separate composite custom render pass of the specified primitives. */
+	CompositeRenderPass,
+	/** Direct composite texture, usually a media texture. */
+	CompositeTexture
+};
+
+namespace UE
+{
+	namespace CompositeCore
+	{
+		struct COMPOSITECORE_API FBuiltInRenderPassOptions final
+		{
+			/** Constructor. */
+			FBuiltInRenderPassOptions();
+
+			/* Custom user flags value used to alter materials in the composite render pass. */
+			TOptional<int32> ViewUserFlagsOverride;
+
+			/** Enables the development shader debug feature that routes the Base Color output to Emissive for the separate render. Non-shipping PC build only. */
+			bool bEnableUnlitViewmode;
+
+			/* Enables the dilation pass with size > 0 (currently only 1 is supported). */
+			int32 OutputDilationSize;
+		};
+
+		/** Render-thread struct for render post-processing work per frame. */
+		struct COMPOSITECORE_API FPostRenderWork final
+		{
+			/** Default frame work. */
+			static const FPostRenderWork& GetDefault();
+
+			/** Constructor. */
+			FPostRenderWork();
+
+			/** Input (pre) processing passes. */
+			TArray<TPair<ECompositeCorePassInput, TArray<const ICompositeCorePassProxy*>>> InputProcessing;
+
+			/** Post-processing passes at the specified locations. */
+			TSortedMap<ISceneViewExtension::EPostProcessingPass, TArray<const ICompositeCorePassProxy*>> FramePasses;
+
+			/** Proxy allocator. */
+			TUniquePtr<FSceneRenderingBulkObjectAllocator> FrameAllocator;
+		};
+	}
+}
+
+/** Render-thread pass proxy. */
+class COMPOSITECORE_API ICompositeCorePassProxy
+{
+public:
+	/** Constructor. */
+	ICompositeCorePassProxy(ECompositeCorePassInput Input0)
+		: InputTypes({ Input0 })
+	{ };
+
+	/* Virtual destructor in case deletion is done on interface pointers. */
+	virtual ~ICompositeCorePassProxy() { };
+
+	/* Render-thread add pass method to override. */
+	virtual FScreenPassTexture Add(FRDGBuilder& GraphBuilder, const FSceneView& InView, const FPostProcessMaterialInputs& Inputs, ISceneViewExtension::EPostProcessingPass InLocation) const = 0;
+
+	/* Number of inputs used by the pass. */
+	int32 GetNumInputs() const { return InputTypes.Num(); }
+
+	/* Number of inputs used by the pass. */
+	ECompositeCorePassInput GetPassInput(int32 Index) const { return (Index < GetNumInputs()) ? InputTypes[Index] : ECompositeCorePassInput::Default; };
+
+protected:
+	/* List of pass input types. */
+	TArray<ECompositeCorePassInput> InputTypes;
+};
+
diff --git a/Engine/Plugins/Experimental/Compositing/HoldoutComposite/HoldoutComposite.uplugin b/Engine/Plugins/Experimental/Compositing/HoldoutComposite/HoldoutComposite.uplugin
--- a/Engine/Plugins/Experimental/Compositing/HoldoutComposite/HoldoutComposite.uplugin
+++ b/Engine/Plugins/Experimental/Compositing/HoldoutComposite/HoldoutComposite.uplugin
@@ -3,7 +3,7 @@
 	"Version": 1,
 	"VersionName": "1.0",
 	"FriendlyName": "HoldoutComposite",
-	"Description": "Plugin designed to facilitate alpha holdout compositing primitives after post-processing.",
+	"Description": "Deprecated plugin now redirected to CompositeCore",
 	"Category": "Compositing",
 	"CreatedBy": "Epic Games, Inc.",
 	"CreatedByURL": "https://epicgames.com",
@@ -16,10 +16,11 @@
 	"IsExperimentalVersion": true,
 	"Installed": false,
 	"Modules": [
+	],
+	"Plugins": [
 		{
-			"Name": "HoldoutComposite",
-			"Type": "Runtime",
-			"LoadingPhase": "PostConfigInit"
+			"Name": "CompositeCore",
+			"Enabled": true
 		}
 	]
 }
\ No newline at end of file
diff --git a/Engine/Plugins/Experimental/Compositing/HoldoutComposite/Shaders/Private/HoldoutCompositeDilate.usf b/Engine/Plugins/Experimental/Compositing/HoldoutComposite/Shaders/Private/HoldoutCompositeDilate.usf
deleted file mode 100644
diff --git a/Engine/Plugins/Experimental/Compositing/HoldoutComposite/Shaders/Private/HoldoutCompositeGamma.usf b/Engine/Plugins/Experimental/Compositing/HoldoutComposite/Shaders/Private/HoldoutCompositeGamma.usf
deleted file mode 100644
diff --git a/Engine/Plugins/Experimental/Compositing/HoldoutComposite/Shaders/Private/HoldoutCompositeMerge.usf b/Engine/Plugins/Experimental/Compositing/HoldoutComposite/Shaders/Private/HoldoutCompositeMerge.usf
deleted file mode 100644
diff --git a/Engine/Plugins/Experimental/Compositing/HoldoutComposite/Source/HoldoutComposite/HoldoutComposite.Build.cs b/Engine/Plugins/Experimental/Compositing/HoldoutComposite/Source/HoldoutComposite/HoldoutComposite.Build.cs
deleted file mode 100644
diff --git a/Engine/Plugins/Experimental/Compositing/HoldoutComposite/Source/HoldoutComposite/Private/HoldoutComposite.cpp b/Engine/Plugins/Experimental/Compositing/HoldoutComposite/Source/HoldoutComposite/Private/HoldoutComposite.cpp
deleted file mode 100644
diff --git a/Engine/Plugins/Experimental/Compositing/HoldoutComposite/Source/HoldoutComposite/Private/HoldoutCompositeComponent.cpp b/Engine/Plugins/Experimental/Compositing/HoldoutComposite/Source/HoldoutComposite/Private/HoldoutCompositeComponent.cpp
deleted file mode 100644
diff --git a/Engine/Plugins/Experimental/Compositing/HoldoutComposite/Source/HoldoutComposite/Private/HoldoutCompositeModule.cpp b/Engine/Plugins/Experimental/Compositing/HoldoutComposite/Source/HoldoutComposite/Private/HoldoutCompositeModule.cpp
deleted file mode 100644
diff --git a/Engine/Plugins/Experimental/Compositing/HoldoutComposite/Source/HoldoutComposite/Private/HoldoutCompositeSceneViewExtension.cpp b/Engine/Plugins/Experimental/Compositing/HoldoutComposite/Source/HoldoutComposite/Private/HoldoutCompositeSceneViewExtension.cpp
deleted file mode 100644
diff --git a/Engine/Plugins/Experimental/Compositing/HoldoutComposite/Source/HoldoutComposite/Private/HoldoutCompositeSceneViewExtension.h b/Engine/Plugins/Experimental/Compositing/HoldoutComposite/Source/HoldoutComposite/Private/HoldoutCompositeSceneViewExtension.h
deleted file mode 100644
diff --git a/Engine/Plugins/Experimental/Compositing/HoldoutComposite/Source/HoldoutComposite/Private/HoldoutCompositeSettings.cpp b/Engine/Plugins/Experimental/Compositing/HoldoutComposite/Source/HoldoutComposite/Private/HoldoutCompositeSettings.cpp
deleted file mode 100644
diff --git a/Engine/Plugins/Experimental/Compositing/HoldoutComposite/Source/HoldoutComposite/Private/HoldoutCompositeSubsystem.cpp b/Engine/Plugins/Experimental/Compositing/HoldoutComposite/Source/HoldoutComposite/Private/HoldoutCompositeSubsystem.cpp
deleted file mode 100644
diff --git a/Engine/Plugins/Experimental/Compositing/HoldoutComposite/Source/HoldoutComposite/Private/HoldoutCompositeViewProjectionComponent.cpp b/Engine/Plugins/Experimental/Compositing/HoldoutComposite/Source/HoldoutComposite/Private/HoldoutCompositeViewProjectionComponent.cpp
deleted file mode 100644
diff --git a/Engine/Plugins/Experimental/Compositing/HoldoutComposite/Source/HoldoutComposite/Private/Passes/HoldoutCompositePassBase.cpp b/Engine/Plugins/Experimental/Compositing/HoldoutComposite/Source/HoldoutComposite/Private/Passes/HoldoutCompositePassBase.cpp
deleted file mode 100644
diff --git a/Engine/Plugins/Experimental/Compositing/HoldoutComposite/Source/HoldoutComposite/Private/Passes/HoldoutCompositePassFXAA.cpp b/Engine/Plugins/Experimental/Compositing/HoldoutComposite/Source/HoldoutComposite/Private/Passes/HoldoutCompositePassFXAA.cpp
deleted file mode 100644
diff --git a/Engine/Plugins/Experimental/Compositing/HoldoutComposite/Source/HoldoutComposite/Private/Passes/HoldoutCompositePassMaterial.cpp b/Engine/Plugins/Experimental/Compositing/HoldoutComposite/Source/HoldoutComposite/Private/Passes/HoldoutCompositePassMaterial.cpp
deleted file mode 100644
diff --git a/Engine/Plugins/Experimental/Compositing/HoldoutComposite/Source/HoldoutComposite/Private/Passes/HoldoutCompositePassMerge.cpp b/Engine/Plugins/Experimental/Compositing/HoldoutComposite/Source/HoldoutComposite/Private/Passes/HoldoutCompositePassMerge.cpp
deleted file mode 100644
diff --git a/Engine/Plugins/Experimental/Compositing/HoldoutComposite/Source/HoldoutComposite/Public/HoldoutComposite.h b/Engine/Plugins/Experimental/Compositing/HoldoutComposite/Source/HoldoutComposite/Public/HoldoutComposite.h
deleted file mode 100644
diff --git a/Engine/Plugins/Experimental/Compositing/HoldoutComposite/Source/HoldoutComposite/Public/HoldoutCompositeComponent.h b/Engine/Plugins/Experimental/Compositing/HoldoutComposite/Source/HoldoutComposite/Public/HoldoutCompositeComponent.h
deleted file mode 100644
diff --git a/Engine/Plugins/Experimental/Compositing/HoldoutComposite/Source/HoldoutComposite/Public/HoldoutCompositeModule.h b/Engine/Plugins/Experimental/Compositing/HoldoutComposite/Source/HoldoutComposite/Public/HoldoutCompositeModule.h
deleted file mode 100644
diff --git a/Engine/Plugins/Experimental/Compositing/HoldoutComposite/Source/HoldoutComposite/Public/HoldoutCompositeSettings.h b/Engine/Plugins/Experimental/Compositing/HoldoutComposite/Source/HoldoutComposite/Public/HoldoutCompositeSettings.h
deleted file mode 100644
diff --git a/Engine/Plugins/Experimental/Compositing/HoldoutComposite/Source/HoldoutComposite/Public/HoldoutCompositeSubsystem.h b/Engine/Plugins/Experimental/Compositing/HoldoutComposite/Source/HoldoutComposite/Public/HoldoutCompositeSubsystem.h
deleted file mode 100644
diff --git a/Engine/Plugins/Experimental/Compositing/HoldoutComposite/Source/HoldoutComposite/Public/HoldoutCompositeViewProjectionComponent.h b/Engine/Plugins/Experimental/Compositing/HoldoutComposite/Source/HoldoutComposite/Public/HoldoutCompositeViewProjectionComponent.h
deleted file mode 100644
diff --git a/Engine/Plugins/Experimental/Compositing/HoldoutComposite/Source/HoldoutComposite/Public/Passes/HoldoutCompositePassBase.h b/Engine/Plugins/Experimental/Compositing/HoldoutComposite/Source/HoldoutComposite/Public/Passes/HoldoutCompositePassBase.h
deleted file mode 100644
diff --git a/Engine/Plugins/Experimental/Compositing/HoldoutComposite/Source/HoldoutComposite/Public/Passes/HoldoutCompositePassFXAA.h b/Engine/Plugins/Experimental/Compositing/HoldoutComposite/Source/HoldoutComposite/Public/Passes/HoldoutCompositePassFXAA.h
deleted file mode 100644
diff --git a/Engine/Plugins/Experimental/Compositing/HoldoutComposite/Source/HoldoutComposite/Public/Passes/HoldoutCompositePassMaterial.h b/Engine/Plugins/Experimental/Compositing/HoldoutComposite/Source/HoldoutComposite/Public/Passes/HoldoutCompositePassMaterial.h
deleted file mode 100644
diff --git a/Engine/Plugins/Experimental/Compositing/HoldoutComposite/Source/HoldoutComposite/Public/Passes/HoldoutCompositePassMerge.h b/Engine/Plugins/Experimental/Compositing/HoldoutComposite/Source/HoldoutComposite/Public/Passes/HoldoutCompositePassMerge.h
deleted file mode 100644
diff --git a/Engine/Plugins/Media/MediaPlate/MediaPlate.uplugin b/Engine/Plugins/Media/MediaPlate/MediaPlate.uplugin
--- a/Engine/Plugins/Media/MediaPlate/MediaPlate.uplugin
+++ b/Engine/Plugins/Media/MediaPlate/MediaPlate.uplugin
@@ -46,7 +46,7 @@
 			"Enabled": true
 		},
 		{
-			"Name": "HoldoutComposite",
+			"Name": "CompositeCore",
 			"Enabled": true
 		}
 	]
diff --git a/Engine/Plugins/Media/MediaPlate/Source/MediaPlate/MediaPlate.Build.cs b/Engine/Plugins/Media/MediaPlate/Source/MediaPlate/MediaPlate.Build.cs
--- a/Engine/Plugins/Media/MediaPlate/Source/MediaPlate/MediaPlate.Build.cs
+++ b/Engine/Plugins/Media/MediaPlate/Source/MediaPlate/MediaPlate.Build.cs
@@ -23,7 +23,7 @@ public class MediaPlate : ModuleRules
 				"Engine",
 				"MediaAssets",
 				"MediaUtils",
-				"HoldoutComposite"
+				"CompositeCore"
 			}
 			);
 
diff --git a/Engine/Plugins/Media/MediaPlate/Source/MediaPlate/Private/MediaPlate.cpp b/Engine/Plugins/Media/MediaPlate/Source/MediaPlate/Private/MediaPlate.cpp
--- a/Engine/Plugins/Media/MediaPlate/Source/MediaPlate/Private/MediaPlate.cpp
+++ b/Engine/Plugins/Media/MediaPlate/Source/MediaPlate/Private/MediaPlate.cpp
@@ -6,7 +6,7 @@
 #include "Engine/StaticMesh.h"
 #include "Engine/World.h"
 #include "HoldoutCompositeComponent.h"
-#include "HoldoutCompositeSubsystem.h"
+#include "CompositeCoreSubsystem.h"
 #include "Materials/Material.h"
 #include "Materials/MaterialInstanceConstant.h"
 #include "Materials/MaterialInstanceDynamic.h"
@@ -169,18 +169,18 @@ void AMediaPlate::SetHoldoutCompositeEnabled(bool bInEnabled)
 {
 	bEnableHoldoutComposite = bInEnabled;
 
-	UHoldoutCompositeSubsystem* HoldoutCompositeSubsystem = UWorld::GetSubsystem<UHoldoutCompositeSubsystem>(GetWorld());
-	if (IsValid(HoldoutCompositeSubsystem))
+	UCompositeCoreSubsystem* CompositeSubsystem = UWorld::GetSubsystem<UCompositeCoreSubsystem>(GetWorld());
+	if (IsValid(CompositeSubsystem))
 	{
 		if (bEnableHoldoutComposite)
 		{
 			// Note: re-registering the same component is safe.
-			HoldoutCompositeSubsystem->RegisterPrimitive(StaticMeshComponent);
+			CompositeSubsystem->RegisterPrimitive(StaticMeshComponent);
 		}
 		else
 		{
 			// Note: destroyed components are also automatically removed by the system.
-			HoldoutCompositeSubsystem->UnregisterPrimitive(StaticMeshComponent);
+			CompositeSubsystem->UnregisterPrimitive(StaticMeshComponent);
 		}
 	}
 }
diff --git a/Engine/Plugins/Runtime/nDisplay/Source/DisplayCluster/DisplayCluster.Build.cs b/Engine/Plugins/Runtime/nDisplay/Source/DisplayCluster/DisplayCluster.Build.cs
--- a/Engine/Plugins/Runtime/nDisplay/Source/DisplayCluster/DisplayCluster.Build.cs
+++ b/Engine/Plugins/Runtime/nDisplay/Source/DisplayCluster/DisplayCluster.Build.cs
@@ -27,7 +27,7 @@ public class DisplayCluster : ModuleRules
 		PrivateDependencyModuleNames.AddRange(
 			new string[] {
 				"HeadMountedDisplay",
-				"HoldoutComposite",
+				"CompositeCore",
 				"InputCore",
 				"Json",
 				"JsonUtilities",
diff --git a/Engine/Plugins/Runtime/nDisplay/Source/DisplayCluster/Private/Game/EngineClasses/Scene/DisplayClusterRootActor.cpp b/Engine/Plugins/Runtime/nDisplay/Source/DisplayCluster/Private/Game/EngineClasses/Scene/DisplayClusterRootActor.cpp
--- a/Engine/Plugins/Runtime/nDisplay/Source/DisplayCluster/Private/Game/EngineClasses/Scene/DisplayClusterRootActor.cpp
+++ b/Engine/Plugins/Runtime/nDisplay/Source/DisplayCluster/Private/Game/EngineClasses/Scene/DisplayClusterRootActor.cpp
@@ -53,7 +53,7 @@
 #include "Components/DisplayClusterStageIsosphereComponent.h"
 #include "Components/LineBatchComponent.h"
 
-#include "HoldoutCompositeSubsystem.h"
+#include "CompositeCoreSubsystem.h"
 #include "UObject/Package.h"
 
 namespace UE::DisplayCluster::RootActor
@@ -1267,7 +1267,7 @@ void ADisplayClusterRootActor::TickPreviewRenderer()
 
 	// The HoldoutComposite plugin requires special project settings.
 	if (NewPreviewSettings.bPreviewEnableHoldoutComposite
-	&& !UHoldoutCompositeSubsystem::IsProjectSettingsValid())
+	&& !UCompositeCoreSubsystem::IsProjectSettingsValid())
 	{
 		NewPreviewSettings.bPreviewEnableHoldoutComposite = false;
 
diff --git a/Engine/Plugins/Runtime/nDisplay/Source/DisplayCluster/Private/Render/Viewport/Preview/DisplayClusterViewportPreviewMesh.cpp b/Engine/Plugins/Runtime/nDisplay/Source/DisplayCluster/Private/Render/Viewport/Preview/DisplayClusterViewportPreviewMesh.cpp
--- a/Engine/Plugins/Runtime/nDisplay/Source/DisplayCluster/Private/Render/Viewport/Preview/DisplayClusterViewportPreviewMesh.cpp
+++ b/Engine/Plugins/Runtime/nDisplay/Source/DisplayCluster/Private/Render/Viewport/Preview/DisplayClusterViewportPreviewMesh.cpp
@@ -9,7 +9,7 @@
 #include "Components/DisplayClusterCameraComponent.h"
 #include "Components/StaticMeshComponent.h"
 
-#include "HoldoutCompositeSubsystem.h"
+#include "CompositeCoreSubsystem.h"
 
 namespace UE::DisplayCluster::ViewportPreviewMesh
 {
@@ -64,13 +64,13 @@ namespace UE::DisplayCluster::ViewportPreviewMesh
 		const EDisplayClusterDisplayDeviceMeshType InMeshType,
 		const EDisplayClusterDisplayDeviceMaterialType InMaterialType = EDisplayClusterDisplayDeviceMaterialType::PreviewMeshMaterial)
 	{
-		if (!IsValid(InMeshComponent) || !UHoldoutCompositeSubsystem::IsProjectSettingsValid())
+		if (!IsValid(InMeshComponent) || !UCompositeCoreSubsystem::IsProjectSettingsValid())
 		{
 			return;
 		}
 
-		UHoldoutCompositeSubsystem* HoldoutComposite = UWorld::GetSubsystem<UHoldoutCompositeSubsystem>(Configuration.GetCurrentWorld());
-		if (!IsValid(HoldoutComposite))
+		UCompositeCoreSubsystem* CompositeSubsystem = UWorld::GetSubsystem<UCompositeCoreSubsystem>(Configuration.GetCurrentWorld());
+		if (!IsValid(CompositeSubsystem))
 		{
 			return;
 		}
@@ -84,13 +84,13 @@ namespace UE::DisplayCluster::ViewportPreviewMesh
 
 		if (bUseHoldout && Configuration.GetRenderFrameSettings().ShouldUseHoldout())
 		{
-			// Render this mesh using Holdout
-			HoldoutComposite->RegisterPrimitive(InMeshComponent);
+			// Render this mesh using into default (holdout) composite pipeline
+			CompositeSubsystem->RegisterPrimitive(InMeshComponent);
 		}
 		else
 		{
 			// Use normal rendering
-			HoldoutComposite->UnregisterPrimitive(InMeshComponent);
+			CompositeSubsystem->UnregisterPrimitive(InMeshComponent);
 		}
 	}
 };
diff --git a/Engine/Plugins/Runtime/nDisplay/nDisplay.uplugin b/Engine/Plugins/Runtime/nDisplay/nDisplay.uplugin
--- a/Engine/Plugins/Runtime/nDisplay/nDisplay.uplugin
+++ b/Engine/Plugins/Runtime/nDisplay/nDisplay.uplugin
@@ -305,7 +305,7 @@
 			"Enabled": true
 		},
 		{
-			"Name": "HoldoutComposite",
+			"Name": "CompositeCore",
 			"Enabled": true
 		}
 	]
