Composite: Support for multiple (external) input textures & subpasses.

Some demos can now be fully represented by a tree of built-in passes, without any PPM.

Additionals changes include:
* Renamed bAlphaIsOpacity as bInvertedAlpha
* Cleanup of scene view extension, submodule & pass proxy
* Removal of the single (external) composite texture: now replaced with channels captured as external texture inputs on the render work struct.
* Moved built-in CRP to `BeginRenderViewFamily` for supporting split screen testing.
* Renamed ICompositeCorePassProxy to FCompositeCorePassProxy since it is not a header-only interface.

Support for processing external inputs to follow later.

diff --git a/Engine/Plugins/Experimental/Compositing/CompositeCore/Shaders/Private/CompositeCoreMerge.usf b/Engine/Plugins/Experimental/Compositing/CompositeCore/Shaders/Private/CompositeCoreMerge.usf
--- a/Engine/Plugins/Experimental/Compositing/CompositeCore/Shaders/Private/CompositeCoreMerge.usf
+++ b/Engine/Plugins/Experimental/Compositing/CompositeCore/Shaders/Private/CompositeCoreMerge.usf
@@ -27,7 +27,7 @@ SCREEN_PASS_TEXTURE_VIEWPORT(Output)
 #define INPUT_PASS_TEXTURE(StructName) \
 	INPUT_PASS_TEXTURE_MEMBER(StructName, Texture2D, Texture) \
 	INPUT_PASS_TEXTURE_MEMBER(StructName, SamplerState, Sampler) \
-    INPUT_PASS_TEXTURE_MEMBER(StructName, uint, bAlphaIsOpacity) \
+    INPUT_PASS_TEXTURE_MEMBER(StructName, uint, bInvertedAlpha) \
     INPUT_PASS_TEXTURE_MEMBER(StructName, uint, SourceEncoding) \
     INPUT_PASS_TEXTURE_MEMBER(StructName, uint, DistortionUV)
     
@@ -107,9 +107,9 @@ float2 CalcInputUVs(float2 ViewportUV, float2 ViewportSize, float2 ViewportMin,
     return (ViewportUV * ViewportSize + ViewportMin) * ExtentInverse;
 }
 
-float GetOpacity(float InAlpha, uint bAlphaIsOpacity)
+float GetOpacity(float InAlpha, uint bInvertedAlpha)
 {
-	return select(bAlphaIsOpacity > 0, InAlpha, 1.0 - InAlpha);
+	return select(bInvertedAlpha > 0, 1.0 - InAlpha, InAlpha);
 }
 
 float SafeDiv(float A, float B)
@@ -147,8 +147,8 @@ void MainPS(
     Color1.rgb = Decode(Color1.rgb, Tex1_SourceEncoding);
 
     // Normalize alpha as opacity
-    Color0.a = GetOpacity(Color0.a, Tex0_bAlphaIsOpacity);
-    Color1.a = GetOpacity(Color1.a, Tex1_bAlphaIsOpacity);
+	Color0.a = GetOpacity(Color0.a, Tex0_bInvertedAlpha);
+	Color1.a = GetOpacity(Color1.a, Tex1_bInvertedAlpha);
 
     {
         // NOTE: This is necessary to extract the solid colors behind translucent alpha holdout masks.
diff --git a/Engine/Plugins/Experimental/Compositing/CompositeCore/Source/CompositeCore/Private/CompositeCoreSceneViewExtension.cpp b/Engine/Plugins/Experimental/Compositing/CompositeCore/Source/CompositeCore/Private/CompositeCoreSceneViewExtension.cpp
--- a/Engine/Plugins/Experimental/Compositing/CompositeCore/Source/CompositeCore/Private/CompositeCoreSceneViewExtension.cpp
+++ b/Engine/Plugins/Experimental/Compositing/CompositeCore/Source/CompositeCore/Private/CompositeCoreSceneViewExtension.cpp
@@ -6,15 +6,11 @@
 
 #include "CommonRenderResources.h"
 #include "Components/PrimitiveComponent.h"
-#include "Containers/Map.h"
 #include "Containers/Set.h"
 #include "EngineUtils.h"
 #include "Engine/Texture.h"
-#include "FXRenderingUtils.h"
 #include "HDRHelper.h"
-#include "Materials/Material.h"
 #include "PixelShaderUtils.h"
-#include "PostProcess/LensDistortion.h"
 #include "PostProcess/PostProcessMaterialInputs.h"
 #include "Rendering/CustomRenderPass.h"
 #include "RenderGraphBuilder.h"
@@ -103,7 +99,7 @@ class FCompositeCoreCustomRenderPass : public FCustomRenderPassBase
 public:
 	IMPLEMENT_CUSTOM_RENDER_PASS(FCompositeCoreCustomRenderPass);
 
-	FCompositeCoreCustomRenderPass(const FIntPoint& InRenderTargetSize, FCompositeCoreSceneViewExtension* InParentExtension, FSceneView& InView, int32 InDilationSize)
+	FCompositeCoreCustomRenderPass(const FIntPoint& InRenderTargetSize, FCompositeCoreSceneViewExtension* InParentExtension, const FSceneView& InView, int32 InDilationSize)
 		: FCustomRenderPassBase(TEXT("CompositeCoreCustomRenderPass"), ERenderMode::DepthAndBasePass, ERenderOutput::SceneColorAndAlpha, InRenderTargetSize)
 		, ParentExtension(InParentExtension)
 		, ViewId(InView.GetViewKey())
@@ -192,11 +188,6 @@ void FCompositeCoreSceneViewExtension::UnregisterPrimitives_GameThread(const TAr
 	}
 }
 
-void FCompositeCoreSceneViewExtension::SetCompositeTextureOverride_GameThread(UTexture* InTextureOverride)
-{
-	CompositeTextureOverride = InTextureOverride;
-}
-
 void FCompositeCoreSceneViewExtension::SetPostRenderWork_GameThread(UE::CompositeCore::FPostRenderWork&& InWork)
 {
 	ENQUEUE_RENDER_COMMAND(CopyCompositeCoreRenderWork)(
@@ -205,6 +196,22 @@ void FCompositeCoreSceneViewExtension::SetPostRenderWork_GameThread(UE::Composit
 			TSharedPtr<FCompositeCoreSceneViewExtension> SVE = StaticCastSharedPtr<FCompositeCoreSceneViewExtension>(WeakThis.Pin());
 			if (SVE.IsValid())
 			{
+				SVE->ExternalInputs_RenderThread.Reset();
+
+				for (int32 Index=0; Index < InWork.ExternalInputs.Num(); ++Index)
+				{
+					TStrongObjectPtr<UTexture> ExternalInput = InWork.ExternalInputs[Index].Pin();
+					
+					if (ExternalInput.IsValid())
+					{
+						if (FTextureResource* TexResource = ExternalInput->GetResource())
+						{
+							const int32 ExternalId = UE::CompositeCore::EXTERNAL_RANGE_START_ID + Index;
+							SVE->ExternalInputs_RenderThread.Add(ExternalId, CreateRenderTarget(TexResource->GetTextureRHI(), TEXT("CompositeExternalInput")));
+						}
+					}
+				}
+
 				SVE->PostRenderWork_RenderThread = MoveTemp(InWork);
 			}
 		});
@@ -219,6 +226,7 @@ void FCompositeCoreSceneViewExtension::ResetPostRenderWork_GameThread()
 			if (SVE.IsValid())
 			{
 				SVE->PostRenderWork_RenderThread.Reset();
+				SVE->ExternalInputs_RenderThread.Reset();
 			}
 		});
 }
@@ -228,6 +236,13 @@ void FCompositeCoreSceneViewExtension::SetBuiltInRenderPassOptions_GameThread(co
 	BuiltInRenderPassOptions = InOptions;
 }
 
+/* Called by the custom render pass to store its view render target for this frame. */
+
+void FCompositeCoreSceneViewExtension::CollectCustomRenderTarget(uint32 InViewId, const TRefCountPtr<IPooledRenderTarget>& InRenderTarget)
+{
+	CustomRenderTargetPerView_RenderThread.Add(InViewId, InRenderTarget);
+}
+
 bool FCompositeCoreSceneViewExtension::IsActiveThisFrame_Internal(const FSceneViewExtensionContext& Context) const
 {
 	bool bIsActive = FWorldSceneViewExtension::IsActiveThisFrame_Internal(Context);
@@ -270,6 +285,10 @@ void FCompositeCoreSceneViewExtension::SetupViewFamily(FSceneViewFamily& InViewF
 }
 
 void FCompositeCoreSceneViewExtension::SetupView(FSceneViewFamily& InViewFamily, FSceneView& InView)
+{
+}
+
+void FCompositeCoreSceneViewExtension::BeginRenderViewFamily(FSceneViewFamily& InViewFamily)
 {
 	// Disable custom render pass for view families without primitive alpha holdout support.
 	if (InViewFamily.EngineShowFlags.AllowPrimitiveAlphaHoldout == 0)
@@ -280,83 +299,88 @@ void FCompositeCoreSceneViewExtension::SetupView(FSceneViewFamily& InViewFamily,
 	const TWeakObjectPtr<UWorld> WorldPtr = GetWorld();
 	check(WorldPtr.IsValid());
 
-	TSet<FPrimitiveComponentId> CompositeCorePrimitiveIds;
-	for (const TWeakObjectPtr<UPrimitiveComponent>& PrimitivePtr : CompositePrimitives)
+	for (int32 ViewIndex = 0; ViewIndex < InViewFamily.Views.Num(); ++ViewIndex)
 	{
-		TStrongObjectPtr<UPrimitiveComponent> Primitive = PrimitivePtr.Pin();
-		// Collect only those primitives that use the bHoldout flag
-		// The user can directly change this flag outside of this VE.
-		if (Primitive.IsValid() && Primitive->bHoldout)
-		{
-			const FPrimitiveComponentId PrimId = Primitive->GetPrimitiveSceneId();
+		const FSceneView& InView = *InViewFamily.Views[ViewIndex];
 
-			if (InView.ShowOnlyPrimitives.IsSet())
+		TSet<FPrimitiveComponentId> CompositeCorePrimitiveIds;
+		for (const TWeakObjectPtr<UPrimitiveComponent>& PrimitivePtr : CompositePrimitives)
+		{
+			TStrongObjectPtr<UPrimitiveComponent> Primitive = PrimitivePtr.Pin();
+			// Collect only those primitives that use the bHoldout flag
+			// The user can directly change this flag outside of this VE.
+			if (Primitive.IsValid() && Primitive->bHoldout)
 			{
-				if (InView.ShowOnlyPrimitives.GetValue().Contains(PrimId))
+				const FPrimitiveComponentId PrimId = Primitive->GetPrimitiveSceneId();
+
+				if (InView.ShowOnlyPrimitives.IsSet())
+				{
+					if (InView.ShowOnlyPrimitives.GetValue().Contains(PrimId))
+					{
+						CompositeCorePrimitiveIds.Add(Primitive->GetPrimitiveSceneId());
+					}
+				}
+				else if (!InView.HiddenPrimitives.Contains(PrimId))
 				{
 					CompositeCorePrimitiveIds.Add(Primitive->GetPrimitiveSceneId());
 				}
 			}
-			else if (!InView.HiddenPrimitives.Contains(PrimId))
-			{
-				CompositeCorePrimitiveIds.Add(Primitive->GetPrimitiveSceneId());
-			}
 		}
-	}
 
-	if (CompositeCorePrimitiveIds.IsEmpty())
-	{
-		return;
-	}
+		if (CompositeCorePrimitiveIds.IsEmpty())
+		{
+			return;
+		}
 
-	// Extract the custom render target size
-	FIntPoint RenderTargetViewSize;
-	if (InView.PrimaryScreenPercentageMethod == EPrimaryScreenPercentageMethod::TemporalUpscale)
-	{
-		// Note: This is equivalent to FViewInfo::GetSecondaryViewRectSize
-		FIntRect OutputRect = FIntRect(0, 0,
-			FMath::CeilToInt(InView.UnscaledViewRect.Width() * InViewFamily.SecondaryViewFraction),
-			FMath::CeilToInt(InView.UnscaledViewRect.Height() * InViewFamily.SecondaryViewFraction)
+		// Extract the custom render target size
+		FIntPoint RenderTargetViewSize;
+		if (InView.PrimaryScreenPercentageMethod == EPrimaryScreenPercentageMethod::TemporalUpscale)
+		{
+			// Note: This is equivalent to FViewInfo::GetSecondaryViewRectSize
+			FIntRect OutputRect = FIntRect(0, 0,
+				FMath::CeilToInt(InView.UnscaledViewRect.Width() * InViewFamily.SecondaryViewFraction),
+				FMath::CeilToInt(InView.UnscaledViewRect.Height() * InViewFamily.SecondaryViewFraction)
+			);
+			QuantizeSceneBufferSize(OutputRect.Max, RenderTargetViewSize);
+		}
+		else
+		{
+			RenderTargetViewSize = InView.UnscaledViewRect.Size();
+		}
+
+		// Create a new custom render pass to render the composite primitive(s)
+		FCompositeCoreCustomRenderPass* CustomRenderPass = new FCompositeCoreCustomRenderPass(
+			RenderTargetViewSize,
+			this,
+			InView,
+			BuiltInRenderPassOptions.OutputDilationSize
 		);
-		QuantizeSceneBufferSize(OutputRect.Max, RenderTargetViewSize);
-	}
-	else
-	{
-		RenderTargetViewSize = InView.UnscaledViewRect.Size();
-	}
 
-	// Create a new custom render pass to render the composite primitive(s)
-	FCompositeCoreCustomRenderPass* CustomRenderPass = new FCompositeCoreCustomRenderPass(
-		RenderTargetViewSize,
-		this,
-		InView,
-		BuiltInRenderPassOptions.OutputDilationSize
-	);
-
-	FSceneInterface::FCustomRenderPassRendererInput PassInput{};
-	PassInput.EngineShowFlags = InViewFamily.EngineShowFlags;
-	PassInput.EngineShowFlags.DisableFeaturesForUnlit();
-	PassInput.EngineShowFlags.SetTranslucency(true);
-	PassInput.EngineShowFlags.SetUnlitViewmode(BuiltInRenderPassOptions.bEnableUnlitViewmode);
-	PassInput.EngineShowFlags.SetAllowPrimitiveAlphaHoldout(false);
-	if (BuiltInRenderPassOptions.ViewUserFlagsOverride.IsSet())
-	{
-		PassInput.bOverridesPostVolumeUserFlags = true;
-		PassInput.PostVolumeUserFlags = BuiltInRenderPassOptions.ViewUserFlagsOverride.GetValue();
+		FSceneInterface::FCustomRenderPassRendererInput PassInput{};
+		PassInput.EngineShowFlags = InViewFamily.EngineShowFlags;
+		PassInput.EngineShowFlags.DisableFeaturesForUnlit();
+		PassInput.EngineShowFlags.SetTranslucency(true);
+		PassInput.EngineShowFlags.SetUnlitViewmode(BuiltInRenderPassOptions.bEnableUnlitViewmode);
+		PassInput.EngineShowFlags.SetAllowPrimitiveAlphaHoldout(false);
+		if (BuiltInRenderPassOptions.ViewUserFlagsOverride.IsSet())
+		{
+			PassInput.bOverridesPostVolumeUserFlags = true;
+			PassInput.PostVolumeUserFlags = BuiltInRenderPassOptions.ViewUserFlagsOverride.GetValue();
+		}
+		// Note: Incoming view location is invalid for scene captures
+		PassInput.ViewLocation = InView.ViewMatrices.GetViewOrigin();
+		PassInput.ViewRotationMatrix = InView.ViewMatrices.GetViewMatrix().RemoveTranslation();
+		PassInput.ViewRotationMatrix.RemoveScaling();
+
+		// Note: Projection matrix here is without jitter, GetProjectionNoAAMatrix() is invalid (not yet available).
+		PassInput.ProjectionMatrix = InView.ViewMatrices.GetProjectionMatrix();
+		PassInput.ViewActor = InView.ViewActor;
+		PassInput.ShowOnlyPrimitives = CompositeCorePrimitiveIds;
+		PassInput.CustomRenderPass = CustomRenderPass;
+		PassInput.bIsSceneCapture = true;
+
+		WorldPtr.Get()->Scene->AddCustomRenderPass(&InViewFamily, PassInput);
 	}
-	// Note: Incoming view location is invalid for scene captures
-	PassInput.ViewLocation = InView.ViewMatrices.GetViewOrigin();
-	PassInput.ViewRotationMatrix = InView.ViewMatrices.GetViewMatrix().RemoveTranslation();
-	PassInput.ViewRotationMatrix.RemoveScaling();
-
-	// Note: Projection matrix here is without jitter, GetProjectionNoAAMatrix() is invalid (not yet available).
-	PassInput.ProjectionMatrix = InView.ViewMatrices.GetProjectionMatrix();
-	PassInput.ViewActor = InView.ViewActor;
-	PassInput.ShowOnlyPrimitives = CompositeCorePrimitiveIds;
-	PassInput.CustomRenderPass = CustomRenderPass;
-	PassInput.bIsSceneCapture = true;
-
-	WorldPtr.Get()->Scene->AddCustomRenderPass(&InViewFamily, PassInput);
 }
 
 void FCompositeCoreSceneViewExtension::PreRenderViewFamily_RenderThread(FRDGBuilder& GraphBuilder, FSceneViewFamily& InViewFamily)
@@ -377,198 +401,102 @@ void FCompositeCoreSceneViewExtension::SubscribeToPostProcessingPass(EPostProces
 	{
 		return;
 	}
-	
-	if (PassId == EPostProcessingPass::BeforeDOF)
-	{
-		// We use BeforeDOF as an input processing step for simplicity, since PrePostProcessPass_RenderThread has a different interface.
-		InOutPassCallbacks.Add(FAfterPassCallbackDelegate::CreateRaw(this, &FCompositeCoreSceneViewExtension::InputProcessing_RenderThread));
-	}
-	if (PassId == EPostProcessingPass::SSRInput && GetRenderWork().FramePasses.Contains(ISceneViewExtension::EPostProcessingPass::SSRInput))
-	{
-		InOutPassCallbacks.Add(FAfterPassCallbackDelegate::CreateRaw(this, &FCompositeCoreSceneViewExtension::SSRInputWork_RenderThread));
-	}
-	else if (PassId == EPostProcessingPass::Tonemap && GetRenderWork().FramePasses.Contains(ISceneViewExtension::EPostProcessingPass::Tonemap))
-	{
-		InOutPassCallbacks.Add(FAfterPassCallbackDelegate::CreateRaw(this, &FCompositeCoreSceneViewExtension::TonemapWork_RenderThread));
-	}
-	else if (PassId == EPostProcessingPass::VisualizeDepthOfField && GetRenderWork().FramePasses.Contains(ISceneViewExtension::EPostProcessingPass::VisualizeDepthOfField))
+
+	if (GetRenderWork().FramePasses.Contains(PassId))
 	{
-		InOutPassCallbacks.Add(FAfterPassCallbackDelegate::CreateRaw(this, &FCompositeCoreSceneViewExtension::EndOfPostProcessing_RenderThread));
+		InOutPassCallbacks.Add(FAfterPassCallbackDelegate::CreateRaw(this, &FCompositeCoreSceneViewExtension::PostProcessWork_RenderThread, PassId));
 	}
 }
 
 bool FCompositeCoreSceneViewExtension::IsActiveForView(const FSceneView& InView) const
 {
-	constexpr bool bEvenIfPendingKill = false;
-	constexpr bool bThreadSafeTest = true;
-
 	bool bIsActive = InView.Family->bIsMainViewFamily;
 	bIsActive &= !InView.bIsSceneCapture;
 	bIsActive &= static_cast<bool>(InView.Family->EngineShowFlags.AllowPrimitiveAlphaHoldout);
-	bIsActive &= CustomRenderTargetPerView_RenderThread.Contains(InView.GetViewKey()) || CompositeTextureOverride.IsValid(bEvenIfPendingKill, bThreadSafeTest);
+	bIsActive &= CustomRenderTargetPerView_RenderThread.Contains(InView.GetViewKey()) || !ExternalInputs_RenderThread.IsEmpty();
 
 	return bIsActive;
 }
 
-FScreenPassTexture FCompositeCoreSceneViewExtension::InputProcessing_RenderThread(FRDGBuilder& GraphBuilder, const FSceneView& InView, const FPostProcessMaterialInputs& Inputs)
+TSortedMap<int32, FScreenPassTextureSlice> FCompositeCoreSceneViewExtension::CreateInputOverridesMap(FRDGBuilder& GraphBuilder, const FSceneView& InView) const
 {
-	using namespace UE::CompositeCore::Private;
+	using namespace UE::CompositeCore;
 
-	// Reset composite textures for view.
-	// NOTE: all post-processing SVE callbacks occur on *the same view*, so this is allowed.
-	CompositeRenderPass_RenderThread = CompositeTexture_RenderThread = FScreenPassTexture();
+	TSortedMap<int32, FScreenPassTextureSlice> InputOverrides;
+	InputOverrides.Reserve(2 + ExternalInputs_RenderThread.Num());
 
-	// Prepare texture inputs for following post-processing passes
 	{
-		const TRefCountPtr<IPooledRenderTarget>* CompositeRenderPassPtr = CustomRenderTargetPerView_RenderThread.Find(InView.GetViewKey());
-		if (CompositeRenderPassPtr != nullptr)
-		{
-			CompositeRenderPass_RenderThread = FScreenPassTexture{ GraphBuilder.RegisterExternalTexture(*CompositeRenderPassPtr) };
-		}
-
-		constexpr bool bEvenIfPendingKill = false;
-		constexpr bool bThreadSafeTest = true;
-		if (CompositeTextureOverride.IsValid(bEvenIfPendingKill, bThreadSafeTest))
-		{
-			if (FTextureResource* TextureResource = CompositeTextureOverride.GetEvenIfUnreachable()->GetResource())
-			{
-				const FTextureRHIRef& TextureRHI = TextureResource->GetTextureRHI();
-				if (TextureRHI.IsValid())
-				{
-					FRDGTextureRef CompositeTextureRDG = GraphBuilder.FindExternalTexture(TextureRHI.GetReference());
-
-					if (CompositeTextureRDG == nullptr)
-					{
-						CompositeTextureRDG = GraphBuilder.RegisterExternalTexture(CreateRenderTarget(TextureRHI.GetReference(), TEXT("CompositeTextureOverride")));
-					}
-					
-					FScreenPassTexture Input{ CompositeTextureRDG };
-					
-					const FScreenPassTextureSlice SceneColorSlice = Inputs.GetInput(EPostProcessMaterialInput::SceneColor);
-					const FScreenPassTextureViewport SceneColorViewport = FScreenPassTextureViewport(SceneColorSlice);
-
-					FRDGTextureDesc OutputDesc = GetPostProcessingDesc(SceneColorSlice.TextureSRV->Desc.Texture->Desc, Inputs.OutputFormat);
-					FScreenPassRenderTarget Output = FScreenPassRenderTarget(GraphBuilder.CreateTexture(OutputDesc, TEXT("CompositeTextureCopy")), SceneColorViewport.Rect, InView.GetOverwriteLoadAction());
-	
-					// TODO: We could avoid the copy if the texture override format is usable as scene color.
-					{
-						const FScreenPassTextureViewport InputViewport(Input);
-						const FScreenPassTextureViewport OutputViewport(Output);
-
-						TShaderMapRef<FCopyRectPS> PixelShader(GetGlobalShaderMap(InView.GetFeatureLevel()));
-
-						FCopyRectPS::FParameters* Parameters = GraphBuilder.AllocParameters<FCopyRectPS::FParameters>();
-						Parameters->InputTexture = Input.Texture;
-						Parameters->InputSampler = TStaticSamplerState<SF_Bilinear>::GetRHI();
-						Parameters->RenderTargets[0] = Output.GetRenderTargetBinding();
-
-						AddDrawScreenPass(GraphBuilder, RDG_EVENT_NAME("DrawTexture"), InView, OutputViewport, InputViewport, PixelShader, Parameters);
-					}
-
-					CompositeTexture_RenderThread = Output;
-				}
-			}
-		}
+		FScreenPassTexture ScreenTex = FScreenPassTexture{ GSystemTextures.GetBlackDummy(GraphBuilder) };
+		InputOverrides.Emplace(BUILT_IN_EMPTY_ID, FScreenPassTextureSlice::CreateFromScreenPassTexture(GraphBuilder, ScreenTex));
 	}
 
-	// Texture processing
-	for (TPair<ECompositeCorePassInput, TArray<const ICompositeCorePassProxy*>> Processing : GetRenderWork().InputProcessing)
+	const TRefCountPtr<IPooledRenderTarget>* CompositeRenderPassPtr = CustomRenderTargetPerView_RenderThread.Find(InView.GetViewKey());
+	if (CompositeRenderPassPtr != nullptr)
 	{
-		const ECompositeCorePassInput ActiveTextureInput = Processing.Key;
-		const TArray<const ICompositeCorePassProxy*>& Passes = Processing.Value;
-
-		if (!ensureMsgf(ActiveTextureInput != ECompositeCorePassInput::Default, TEXT("Processing is expected to run on composite textures.")))
-		{
-			continue;
-		}
-
-		TSortedMap<ECompositeCorePassInput, FScreenPassTextureSlice> TextureInputs;
-		if (CompositeRenderPass_RenderThread.IsValid())
-		{
-			TextureInputs.Emplace(ECompositeCorePassInput::CompositeRenderPass, FScreenPassTextureSlice::CreateFromScreenPassTexture(GraphBuilder, CompositeRenderPass_RenderThread));
-		}
-		if (CompositeTexture_RenderThread.IsValid())
-		{
-			TextureInputs.Emplace(ECompositeCorePassInput::CompositeTexture, FScreenPassTextureSlice::CreateFromScreenPassTexture(GraphBuilder, CompositeTexture_RenderThread));
-		}
-
-		FScreenPassTexture Output;
-
-		for (const ICompositeCorePassProxy* Pass : Passes)
-		{
-			if (FScreenPassTextureSlice* SliceInput = TextureInputs.Find(ActiveTextureInput))
-			{
-				FPostProcessMaterialInputs PassInputs = Inputs;
-				PassInputs.SetInput(EPostProcessMaterialInput::SceneColor, *SliceInput);
-				PassInputs.OverrideOutput = FScreenPassRenderTarget{};
-
-				// Register pass and update output
-				Output = Pass->Add(GraphBuilder, InView, PassInputs, ISceneViewExtension::EPostProcessingPass::BeforeDOF);
-
-				// Update active input for the next pass
-				TextureInputs[ActiveTextureInput] = FScreenPassTextureSlice::CreateFromScreenPassTexture(GraphBuilder, Output);
-			}
-		}
-
-		// Update composite textures for later callbacks in post-processing
-		if (Output.IsValid())
-		{
-			switch (ActiveTextureInput)
-			{
-			case ECompositeCorePassInput::CompositeRenderPass:
-				CompositeRenderPass_RenderThread = Output;
-				break;
-			case ECompositeCorePassInput::CompositeTexture:
-				CompositeTexture_RenderThread = Output;
-				break;
+		const FScreenPassTexture ScreenTex = FScreenPassTexture{ GraphBuilder.RegisterExternalTexture(*CompositeRenderPassPtr) };
+		InputOverrides.Emplace(BUILT_IN_CRP_ID, FScreenPassTextureSlice::CreateFromScreenPassTexture(GraphBuilder, ScreenTex));
+	}
 
-			default:
-				checkNoEntry();
-				break;
-			}
-		}
+	for (const auto& Pair : ExternalInputs_RenderThread)
+	{
+		const FScreenPassTexture ScreenTex = FScreenPassTexture{ GraphBuilder.RegisterExternalTexture(Pair.Value) };
+		InputOverrides.Emplace(Pair.Key, FScreenPassTextureSlice::CreateFromScreenPassTexture(GraphBuilder, ScreenTex));
 	}
 
-	return Inputs.ReturnUntouchedSceneColorForPostProcessing(GraphBuilder);
+	return InputOverrides;
 }
 
-FScreenPassTexture FCompositeCoreSceneViewExtension::PostProcessWork_RenderThread(FRDGBuilder& GraphBuilder, const FSceneView& InView, const FPostProcessMaterialInputs& Inputs, ISceneViewExtension::EPostProcessingPass InLocation)
+bool FCompositeCoreSceneViewExtension::ApplyPasses_Recursive(FRDGBuilder& GraphBuilder, const FSceneView& InView,
+	const FPostProcessMaterialInputs& Inputs, ISceneViewExtension::EPostProcessingPass InLocation,
+	const TArray<const FCompositeCorePassProxy*> InPasses, const TSortedMap<int32, FScreenPassTextureSlice>& InputOverrides,
+	const FIntRect& SceneColorRect, int32 RecursionLevel, FScreenPassTexture& Output)
 {
-	using namespace UE::CompositeCore::Private;
-
-	TSortedMap<ECompositeCorePassInput, FScreenPassTextureSlice> TextureInputs;
-	TextureInputs.Emplace(ECompositeCorePassInput::Default, Inputs.GetInput(EPostProcessMaterialInput::SceneColor));
-	if (CompositeRenderPass_RenderThread.IsValid())
-	{
-		TextureInputs.Emplace(ECompositeCorePassInput::CompositeRenderPass, FScreenPassTextureSlice::CreateFromScreenPassTexture(GraphBuilder, CompositeRenderPass_RenderThread));
-	}
-	if (CompositeTexture_RenderThread.IsValid())
+	if (InPasses.IsEmpty())
 	{
-		TextureInputs.Emplace(ECompositeCorePassInput::CompositeTexture, FScreenPassTextureSlice::CreateFromScreenPassTexture(GraphBuilder, CompositeTexture_RenderThread));
+		return false;
 	}
 
-	FScreenPassTexture Output;
-	const TArray<const ICompositeCorePassProxy*>& Passes = GetRenderWork().FramePasses[InLocation];
-	for (int32 Index = 0; Index < Passes.Num(); ++Index)
+	// Default pass inputs
+	FPostProcessMaterialInputs BasePassInputs = Inputs;
+	
+	// Iterate over all passes
+	for (int32 Index = 0; Index < InPasses.Num(); ++Index)
 	{
-		const ICompositeCorePassProxy* Pass = Passes[Index];
-		const bool bIsLastPass = (Index == Passes.Num() - 1);
+		const FCompositeCorePassProxy* Pass = InPasses[Index];
+		
+		// Output override only applies to the (base) last pass
+		const bool bIsLastPass = (Index == InPasses.Num() - 1) && (RecursionLevel == 0);
 
-		// Update pass inputs
-		FPostProcessMaterialInputs PassInputs = Inputs;
-		int32 ConnectedInputs = 0;
-		for (int32 InputIndex = 0; InputIndex < Pass->GetNumInputs() && InputIndex < TextureInputs.Num(); ++InputIndex)
+		// Update inputs for the current pass
+		FPostProcessMaterialInputs PassInputs = BasePassInputs;
+
+		for (int32 InputIndex = 0; InputIndex < Pass->GetNumInputs(); ++InputIndex)
 		{
-			if (FScreenPassTextureSlice* SliceInput = TextureInputs.Find(Pass->GetPassInput(InputIndex)))
+			if (!ensureMsgf(InputIndex <= static_cast<int32>(EPostProcessMaterialInput::Velocity), TEXT("Inputs exceeding the number of post-process material inputs is currently not allowed.")))
 			{
-				PassInputs.SetInput(static_cast<EPostProcessMaterialInput>(InputIndex), *SliceInput);
-				ConnectedInputs++;
+				break;
 			}
-		}
 
-		if (ConnectedInputs == 0)
-		{
-			continue;
+			// Apply input overrides
+			const UE::CompositeCore::FPassInput& PassInput = Pass->GetPassInput(InputIndex);
+			if (PassInput.IsSet())
+			{
+				if (const FScreenPassTextureSlice* SliceInput = InputOverrides.Find(PassInput.GetValue().Id))
+				{
+					PassInputs.SetInput(static_cast<EPostProcessMaterialInput>(InputIndex), *SliceInput);
+				}
+			}
+
+			// Recursively apply sub-passes per input index, automatically overriding inputs.
+			const TArray<const FCompositeCorePassProxy*>* SubPasses = Pass->GetSubPasses(InputIndex);
+			if (SubPasses)
+			{
+				FScreenPassTexture SubPassOutput;
+				if (ApplyPasses_Recursive(GraphBuilder, InView, Inputs, InLocation, *SubPasses, InputOverrides, SceneColorRect, RecursionLevel + 1, SubPassOutput))
+				{
+					PassInputs.SetInput(static_cast<EPostProcessMaterialInput>(InputIndex), FScreenPassTextureSlice::CreateFromScreenPassTexture(GraphBuilder, SubPassOutput));
+				}
+			}
 		}
 
 		// Only apply the output override on the last pass
@@ -578,12 +506,28 @@ FScreenPassTexture FCompositeCoreSceneViewExtension::PostProcessWork_RenderThrea
 		}
 
 		// Register pass and update output
-		Output = Pass->Add(GraphBuilder, InView, PassInputs, InLocation);
+		Output = Pass->Add(GraphBuilder, InView, PassInputs, SceneColorRect, InLocation);
 
 		// Update input 0 for the next pass
-		TextureInputs[ECompositeCorePassInput::Default] = FScreenPassTextureSlice::CreateFromScreenPassTexture(GraphBuilder, Output);
+		BasePassInputs.SetInput(EPostProcessMaterialInput::SceneColor, FScreenPassTextureSlice::CreateFromScreenPassTexture(GraphBuilder, Output));
 	}
 
+	return true;
+}
+
+FScreenPassTexture FCompositeCoreSceneViewExtension::PostProcessWork_RenderThread(FRDGBuilder& GraphBuilder, const FSceneView& InView, const FPostProcessMaterialInputs& Inputs, ISceneViewExtension::EPostProcessingPass InLocation)
+{
+	using namespace UE::CompositeCore;
+
+	TSortedMap<int32, FScreenPassTextureSlice> InputOverrides = CreateInputOverridesMap(GraphBuilder, InView);
+	const FIntRect& SceneColorRect = Inputs.GetInput(EPostProcessMaterialInput::SceneColor).ViewRect;
+	const TArray<const FCompositeCorePassProxy*>& Passes = GetRenderWork().FramePasses[InLocation];
+
+	FScreenPassTexture Output;
+	
+	constexpr int32 RecursionLevel = 0;
+	ApplyPasses_Recursive(GraphBuilder, InView, Inputs, InLocation, Passes, InputOverrides, SceneColorRect, RecursionLevel, Output);
+
 	if (Output.IsValid())
 	{
 		return MoveTemp(Output);
@@ -594,21 +538,6 @@ FScreenPassTexture FCompositeCoreSceneViewExtension::PostProcessWork_RenderThrea
 	}
 }
 
-FScreenPassTexture FCompositeCoreSceneViewExtension::SSRInputWork_RenderThread(FRDGBuilder& GraphBuilder, const FSceneView& InView, const FPostProcessMaterialInputs& Inputs)
-{
-	return PostProcessWork_RenderThread(GraphBuilder, InView, Inputs, ISceneViewExtension::EPostProcessingPass::SSRInput);
-}
-
-FScreenPassTexture FCompositeCoreSceneViewExtension::TonemapWork_RenderThread(FRDGBuilder& GraphBuilder, const FSceneView& InView, const FPostProcessMaterialInputs& Inputs)
-{
-	return PostProcessWork_RenderThread(GraphBuilder, InView, Inputs, ISceneViewExtension::EPostProcessingPass::Tonemap);
-}
-
-FScreenPassTexture FCompositeCoreSceneViewExtension::EndOfPostProcessing_RenderThread(FRDGBuilder& GraphBuilder, const FSceneView& InView, const FPostProcessMaterialInputs& Inputs)
-{
-	return PostProcessWork_RenderThread(GraphBuilder, InView, Inputs, ISceneViewExtension::EPostProcessingPass::VisualizeDepthOfField);
-}
-
 void FCompositeCoreSceneViewExtension::PostRenderViewFamily_RenderThread(FRDGBuilder& GraphBuilder, FSceneViewFamily& InViewFamily)
 {
 
diff --git a/Engine/Plugins/Experimental/Compositing/CompositeCore/Source/CompositeCore/Private/CompositeCoreSceneViewExtension.h b/Engine/Plugins/Experimental/Compositing/CompositeCore/Source/CompositeCore/Private/CompositeCoreSceneViewExtension.h
--- a/Engine/Plugins/Experimental/Compositing/CompositeCore/Source/CompositeCore/Private/CompositeCoreSceneViewExtension.h
+++ b/Engine/Plugins/Experimental/Compositing/CompositeCore/Source/CompositeCore/Private/CompositeCoreSceneViewExtension.h
@@ -7,7 +7,6 @@
 #include "Containers/ContainersFwd.h"
 #include "RenderGraphFwd.h"
 #include "RendererInterface.h"
-#include "ScreenPass.h"
 
 #include "CompositeCoreSettings.h"
 #include "Passes/CompositeCorePassProxy.h"
@@ -15,6 +14,7 @@
 class FViewInfo;
 struct FScreenPassRenderTarget;
 struct FScreenPassTexture;
+struct FScreenPassTextureSlice;
 class UMaterialInterface;
 class UTexture;
 
@@ -30,9 +30,6 @@ public:
 	/* Unregister primitives for compositing. */
 	void UnregisterPrimitives_GameThread(const TArray<UPrimitiveComponent*>& InPrimitiveComponents);
 
-	/* Set a texture override to be used instead of the separate composited render. */
-	void SetCompositeTextureOverride_GameThread(UTexture* InTextureOverride);
-
 	/* Set post-processing render work to be enqueued on the render thread. */
 	void SetPostRenderWork_GameThread(UE::CompositeCore::FPostRenderWork&& InWork);
 
@@ -43,16 +40,13 @@ public:
 	void SetBuiltInRenderPassOptions_GameThread(const UE::CompositeCore::FBuiltInRenderPassOptions& InOptions);
 
 	/* Called by the custom render pass to store its view render target for this frame. */
-	template<typename T>
-	void CollectCustomRenderTarget(uint32 InViewId, T InRenderTarget)
-	{
-		CustomRenderTargetPerView_RenderThread.Add(InViewId, InRenderTarget);
-	}
+	void CollectCustomRenderTarget(uint32 InViewId, const TRefCountPtr<IPooledRenderTarget>& InRenderTarget);
 
 	//~ Begin ISceneViewExtension Interface
 	virtual int32 GetPriority() const override;
 	virtual void SetupViewFamily(FSceneViewFamily& InViewFamily) override;
 	virtual void SetupView(FSceneViewFamily& InViewFamily, FSceneView& InView) override;
+	virtual void BeginRenderViewFamily(FSceneViewFamily& InViewFamily) override;
 	virtual void PreRenderViewFamily_RenderThread(FRDGBuilder& GraphBuilder, FSceneViewFamily& InViewFamily) override;
 	virtual void PreRenderView_RenderThread(FRDGBuilder& GraphBuilder, FSceneView& InView) override;
 	virtual void PrePostProcessPass_RenderThread(FRDGBuilder& GraphBuilder, const FSceneView& InView, const FPostProcessingInputs& Inputs) override;
@@ -73,12 +67,16 @@ private:
 	/* Active check to avoid work in post-processing. */
 	bool IsActiveForView(const FSceneView& InView) const;
 
-	/* Callbacks for processing passes. */
-	FScreenPassTexture InputProcessing_RenderThread(FRDGBuilder& GraphBuilder, const FSceneView& InView, const FPostProcessMaterialInputs& Inputs);
-	FScreenPassTexture PostProcessWork_RenderThread(FRDGBuilder& GraphBuilder, const FSceneView& InView, const FPostProcessMaterialInputs& Inputs, ISceneViewExtension::EPostProcessingPass InLocation);
-	FScreenPassTexture SSRInputWork_RenderThread(FRDGBuilder& GraphBuilder, const FSceneView& InView, const FPostProcessMaterialInputs& Inputs);
-	FScreenPassTexture TonemapWork_RenderThread(FRDGBuilder& GraphBuilder, const FSceneView& InView, const FPostProcessMaterialInputs& Inputs);
-	FScreenPassTexture EndOfPostProcessing_RenderThread(FRDGBuilder& GraphBuilder, const FSceneView& InView, const FPostProcessMaterialInputs& Inputs);
+	TSortedMap<int32, FScreenPassTextureSlice> CreateInputOverridesMap(FRDGBuilder& GraphBuilder, const FSceneView& InView) const;
+
+	bool ApplyPasses_Recursive(FRDGBuilder& GraphBuilder, const FSceneView& InView,
+		const FPostProcessMaterialInputs& Inputs, ISceneViewExtension::EPostProcessingPass InLocation,
+		const TArray<const FCompositeCorePassProxy*> InPasses, const TSortedMap<int32, FScreenPassTextureSlice>& InputOverrides,
+		const FIntRect& SceneColorRect, int32 Level, FScreenPassTexture& Output);
+
+	/* Callback for processing passes. */
+	FScreenPassTexture PostProcessWork_RenderThread(FRDGBuilder& GraphBuilder, const FSceneView& InView,
+		const FPostProcessMaterialInputs& Inputs, ISceneViewExtension::EPostProcessingPass InLocation);
 
 	/* Collection of primitives to render as a custom render pass and composite after post-processing. .*/
 	TSet<TWeakObjectPtr<UPrimitiveComponent>> CompositePrimitives;
@@ -86,19 +84,13 @@ private:
 	/* Built-in composite render pass options. .*/
 	UE::CompositeCore::FBuiltInRenderPassOptions BuiltInRenderPassOptions;
 
-	/* Composite texture override use directly instead of the separate composited render. */
-	TWeakObjectPtr<UTexture> CompositeTextureOverride;
+	/* Array of external texture inputs. */
+	TMap<UE::CompositeCore::ResourceId, TRefCountPtr<IPooledRenderTarget>> ExternalInputs_RenderThread;
 
 	/* Custom render pass render targets for each active view .*/
 	TMap<uint32, TRefCountPtr<IPooledRenderTarget>> CustomRenderTargetPerView_RenderThread;
 
-	/* Custom actor render work. */
+	/* Custom render work. */
 	TOptional<UE::CompositeCore::FPostRenderWork> PostRenderWork_RenderThread;
-
-	/* Composite render pass screen texture, only valid in post-processing callbacks. */
-	FScreenPassTexture CompositeRenderPass_RenderThread;
-
-	/* Composite screen texture, only valid in post-processing callbacks. */
-	FScreenPassTexture CompositeTexture_RenderThread;
 };
 
diff --git a/Engine/Plugins/Experimental/Compositing/CompositeCore/Source/CompositeCore/Private/CompositeCoreSubsystem.cpp b/Engine/Plugins/Experimental/Compositing/CompositeCore/Source/CompositeCore/Private/CompositeCoreSubsystem.cpp
--- a/Engine/Plugins/Experimental/Compositing/CompositeCore/Source/CompositeCore/Private/CompositeCoreSubsystem.cpp
+++ b/Engine/Plugins/Experimental/Compositing/CompositeCore/Source/CompositeCore/Private/CompositeCoreSubsystem.cpp
@@ -136,14 +136,6 @@ void UCompositeCoreSubsystem::ResetBuiltInRenderPassOptions()
 	}
 }
 
-void UCompositeCoreSubsystem::SetCompositeTextureOverride(UTexture* InTextureOverride)
-{
-	if (CompositeCoreViewExtension.IsValid())
-	{
-		CompositeCoreViewExtension->SetCompositeTextureOverride_GameThread(InTextureOverride);
-	}
-}
-
 bool UCompositeCoreSubsystem::IsProjectSettingsValid()
 {
 	const URendererSettings* RendererSettings = GetDefault<URendererSettings>();
diff --git a/Engine/Plugins/Experimental/Compositing/CompositeCore/Source/CompositeCore/Private/Passes/CompositeCorePassMergeProxy.cpp b/Engine/Plugins/Experimental/Compositing/CompositeCore/Source/CompositeCore/Private/Passes/CompositeCorePassMergeProxy.cpp
--- a/Engine/Plugins/Experimental/Compositing/CompositeCore/Source/CompositeCore/Private/Passes/CompositeCorePassMergeProxy.cpp
+++ b/Engine/Plugins/Experimental/Compositing/CompositeCore/Source/CompositeCore/Private/Passes/CompositeCorePassMergeProxy.cpp
@@ -25,7 +25,7 @@ class FCompositeCoreMergeShader : public FGlobalShader
 	BEGIN_SHADER_PARAMETER_STRUCT(FInputPassTextureParameters, )
 		SHADER_PARAMETER_RDG_TEXTURE(Texture2D, Texture)
 		SHADER_PARAMETER_SAMPLER(SamplerState, Sampler)
-		SHADER_PARAMETER(uint32, bAlphaIsOpacity)
+		SHADER_PARAMETER(uint32, bInvertedAlpha)
 		SHADER_PARAMETER(uint32, SourceEncoding)
 		SHADER_PARAMETER(uint32, DistortionUV)
 	END_SHADER_PARAMETER_STRUCT()
@@ -70,7 +70,7 @@ namespace UE
 
 		FCompositeCoreMergeShader::FInputPassTextureParameters GetPassTextureParameters(
 			const FScreenPassTexture& InTexture,
-			const ECompositeCorePassInput InputType,
+			const UE::CompositeCore::FPassInput PassInput,
 			const FSceneViewFamily* Family,
 			const ISceneViewExtension::EPostProcessingPass InLocation,
 			const bool bUseLensDistortion)
@@ -79,11 +79,11 @@ namespace UE
 
 			Parameters.Texture = InTexture.Texture;
 			Parameters.Sampler = TStaticSamplerState<SF_Bilinear>::GetRHI();
-			Parameters.bAlphaIsOpacity = (InputType == ECompositeCorePassInput::CompositeTexture) ? 1 : 0;
+			Parameters.bInvertedAlpha = PassInput.IsSet() ? PassInput.GetValue().bInvertedAlpha : true;
 			Parameters.SourceEncoding = static_cast<uint32>(ESourceEncoding::Linear);
 
 			// After-tonemap scene color may have encoding manually applied, as opposed to _SRGB textures
-			if (InputType == ECompositeCorePassInput::Default && InLocation >= ISceneViewExtension::EPostProcessingPass::Tonemap)
+			if (!PassInput.IsSet() && InLocation >= ISceneViewExtension::EPostProcessingPass::Tonemap)
 			{
 				if ((Family->EngineShowFlags.Tonemapper == 0) || (Family->EngineShowFlags.PostProcessing == 0))
 				{
@@ -99,24 +99,16 @@ namespace UE
 
 			if (bUseLensDistortion)
 			{
-				const bool bOutputIsDistorted = !(InLocation == ISceneViewExtension::EPostProcessingPass::SSRInput ||
-					InLocation == ISceneViewExtension::EPostProcessingPass::TranslucencyAfterDOF);
+				const bool bInputIsDistorted = PassInput.IsSet() ? PassInput.GetValue().bDistorted
+					: InLocation >= ISceneViewExtension::EPostProcessingPass::SSRInput;
 
-				if (bOutputIsDistorted)
-				{
-					// If output is distorted, we must sample undistorted textures with undistorted UVs
-					if (InputType == ECompositeCorePassInput::CompositeRenderPass)
-					{
-						DistortionUV = EDistortionUV::Undistorted;
-					}
-				}
-				else
+				const bool bOutputIsDistorted = InLocation <= ISceneViewExtension::EPostProcessingPass::SSRInput;
+
+				// Distortion mismatch?
+				if (bInputIsDistorted != bOutputIsDistorted)
 				{
-					// If output is undistorted, we must sample distorted textures with distorted UVs
-					if (InputType == ECompositeCorePassInput::Default || InputType == ECompositeCorePassInput::CompositeTexture)
-					{
-						DistortionUV = EDistortionUV::Distorted;
-					}
+					// Sample according to input distortion
+					DistortionUV = bInputIsDistorted ? EDistortionUV::Distorted : EDistortionUV::Undistorted;
 				}
 			}
 
@@ -125,13 +117,12 @@ namespace UE
 			return Parameters;
 		}
 
-		FMergePassProxy::FMergePassProxy(ECompositeCorePassInput Input0, ECompositeCorePassInput Input1)
-			: ICompositeCorePassProxy(Input0)
+		FMergePassProxy::FMergePassProxy(UE::CompositeCore::FPassInputArray InPassInputs)
+			: FCompositeCorePassProxy(MoveTemp(InPassInputs))
 		{
-			InputTypes.Add(Input1);
 		}
 
-		FScreenPassTexture FMergePassProxy::Add(FRDGBuilder& GraphBuilder, const FSceneView& InView, const FPostProcessMaterialInputs& Inputs, ISceneViewExtension::EPostProcessingPass InLocation) const
+		FScreenPassTexture FMergePassProxy::Add(FRDGBuilder& GraphBuilder, const FSceneView& InView, const FPostProcessMaterialInputs& Inputs, const FIntRect& OutputViewRect, ISceneViewExtension::EPostProcessingPass InLocation) const
 		{
 			RDG_EVENT_SCOPE_STAT(GraphBuilder, FCompositeCoreMerge, "CompositeCore.Merge");
 			RDG_GPU_STAT_SCOPE(GraphBuilder, FCompositeCoreMerge);
@@ -160,7 +151,7 @@ namespace UE
 			FScreenPassRenderTarget Output = Inputs.OverrideOutput;
 			if (!Output.IsValid())
 			{
-				Output = FScreenPassRenderTarget::CreateFromInput(GraphBuilder, Input0, InView.GetOverwriteLoadAction(), TEXT("CompositeCoreMergeOutput"));
+				Output = CreateOutputRenderTarget(GraphBuilder, InView, OutputViewRect, Input0.Texture->Desc, TEXT("CompositeCoreMergeOutput"));
 			}
 
 			const FLensDistortionLUT& LensDistortionLUT = LensDistortion::GetLUTUnsafe(InView);
@@ -171,8 +162,8 @@ namespace UE
 			PassParameters->Input0 = GetScreenPassTextureViewportParameters(FScreenPassTextureViewport(Input0));
 			PassParameters->Input1 = GetScreenPassTextureViewportParameters(FScreenPassTextureViewport(Input1));
 			PassParameters->Output = GetScreenPassTextureViewportParameters(FScreenPassTextureViewport(Output));
-			PassParameters->Tex0 = GetPassTextureParameters(Input0, InputTypes[0], Family, InLocation, bUseLensDistortion);
-			PassParameters->Tex1 = GetPassTextureParameters(Input1, InputTypes[1], Family, InLocation, bUseLensDistortion);
+			PassParameters->Tex0 = GetPassTextureParameters(Input0, PassInputs[0], Family, InLocation, bUseLensDistortion);
+			PassParameters->Tex1 = GetPassTextureParameters(Input1, PassInputs[1], Family, InLocation, bUseLensDistortion);
 			if (bUseLensDistortion)
 			{
 				PassParameters->DistortingDisplacementTexture = LensDistortionLUT.DistortingDisplacementTexture;
diff --git a/Engine/Plugins/Experimental/Compositing/CompositeCore/Source/CompositeCore/Private/Passes/CompositeCorePassProxy.cpp b/Engine/Plugins/Experimental/Compositing/CompositeCore/Source/CompositeCore/Private/Passes/CompositeCorePassProxy.cpp
--- a/Engine/Plugins/Experimental/Compositing/CompositeCore/Source/CompositeCore/Private/Passes/CompositeCorePassProxy.cpp
+++ b/Engine/Plugins/Experimental/Compositing/CompositeCore/Source/CompositeCore/Private/Passes/CompositeCorePassProxy.cpp
@@ -4,6 +4,11 @@
 
 #include "Passes/CompositeCorePassMergeProxy.h"
 
+#include "FXRenderingUtils.h"
+#include "RenderGraphBuilder.h"
+#include "RenderGraphUtils.h"
+#include "ScreenPass.h"
+
 namespace UE
 {
 	namespace CompositeCore
@@ -20,15 +25,21 @@ namespace UE
 			static FPostRenderWork DefaultRenderWork = []() -> FPostRenderWork
 				{
 					FPostRenderWork RenderWork;
+					TArray<const FCompositeCorePassProxy*>& SSRInputWork = RenderWork.FramePasses.Add(ISceneViewExtension::EPostProcessingPass::SSRInput);
+					TArray<const FCompositeCorePassProxy*>& AfterTonemapWork = RenderWork.FramePasses.Add(ISceneViewExtension::EPostProcessingPass::Tonemap);
 
-					constexpr ECompositeCorePassInput Input0 = ECompositeCorePassInput::Default;
-					constexpr ECompositeCorePassInput Input1 = ECompositeCorePassInput::CompositeRenderPass;
-
-					TArray<const ICompositeCorePassProxy*>& SSRInputWork = RenderWork.FramePasses.Add(ISceneViewExtension::EPostProcessingPass::SSRInput);
-					TArray<const ICompositeCorePassProxy*>& AfterTonemapWork = RenderWork.FramePasses.Add(ISceneViewExtension::EPostProcessingPass::Tonemap);
+					FPassInputArray PassInputs;
+					// Input 0: Default
+					PassInputs.Add({});
+					
+					// Input 1: Built-in custom render pass
+					FPassResourceDesc Desc;
+					Desc.Id = BUILT_IN_CRP_ID;
+					Desc.bInvertedAlpha = true;
+					PassInputs.Add(Desc);
 
-					SSRInputWork.Add(RenderWork.FrameAllocator->Create<FMergePassProxy>(Input0, Input1));
-					AfterTonemapWork.Add(RenderWork.FrameAllocator->Create<FMergePassProxy>(Input0, Input1));
+					SSRInputWork.Add(RenderWork.FrameAllocator->Create<FMergePassProxy>(PassInputs));
+					AfterTonemapWork.Add(RenderWork.FrameAllocator->Create<FMergePassProxy>(PassInputs));
 
 					return RenderWork;
 				}();
@@ -37,10 +48,20 @@ namespace UE
 		}
 
 		FPostRenderWork::FPostRenderWork()
-			: InputProcessing{}
-			, FramePasses{}
+			: FramePasses{}
 			, FrameAllocator{new FSceneRenderingBulkObjectAllocator}
 		{ }
 	}
 }
 
+FScreenPassRenderTarget FCompositeCorePassProxy::CreateOutputRenderTarget(FRDGBuilder& GraphBuilder,
+	const FSceneView& InView, const FIntRect& OutputViewRect, FRDGTextureDesc OutputDesc, const TCHAR* InName)
+{
+	OutputDesc.Format	= FSceneTexturesConfig::Get().ColorFormat;
+	OutputDesc.NumMips	= 1;
+	OutputDesc.Depth	= 1;
+	OutputDesc.Flags	= TexCreate_RenderTargetable | TexCreate_ShaderResource; // TexCreate_UAV;
+	OutputDesc.Extent	= OutputViewRect.Size();
+
+	return FScreenPassRenderTarget(GraphBuilder.CreateTexture(OutputDesc, InName), OutputViewRect, InView.GetOverwriteLoadAction());
+}
diff --git a/Engine/Plugins/Experimental/Compositing/CompositeCore/Source/CompositeCore/Public/CompositeCoreSubsystem.h b/Engine/Plugins/Experimental/Compositing/CompositeCore/Source/CompositeCore/Public/CompositeCoreSubsystem.h
--- a/Engine/Plugins/Experimental/Compositing/CompositeCore/Source/CompositeCore/Public/CompositeCoreSubsystem.h
+++ b/Engine/Plugins/Experimental/Compositing/CompositeCore/Source/CompositeCore/Public/CompositeCoreSubsystem.h
@@ -68,10 +68,6 @@ public:
 	/* Reset built-in composite render pass options. */
 	void ResetBuiltInRenderPassOptions();
 
-	/* Set a texture override to be used instead of the separate composited render. */
-	UFUNCTION(BlueprintCallable, Category = "Composite Core")
-	virtual void SetCompositeTextureOverride(UTexture* InTextureOverride);
-
 	/** True if the project settings are valid for the CompositeCore plugin to work. */
 	static bool IsProjectSettingsValid();
 
diff --git a/Engine/Plugins/Experimental/Compositing/CompositeCore/Source/CompositeCore/Public/Passes/CompositeCorePassMergeProxy.h b/Engine/Plugins/Experimental/Compositing/CompositeCore/Source/CompositeCore/Public/Passes/CompositeCorePassMergeProxy.h
--- a/Engine/Plugins/Experimental/Compositing/CompositeCore/Source/CompositeCore/Public/Passes/CompositeCorePassMergeProxy.h
+++ b/Engine/Plugins/Experimental/Compositing/CompositeCore/Source/CompositeCore/Public/Passes/CompositeCorePassMergeProxy.h
@@ -20,13 +20,13 @@ namespace UE
 {
 	namespace CompositeCore
 	{
-		class COMPOSITECORE_API FMergePassProxy : public ICompositeCorePassProxy
+		class COMPOSITECORE_API FMergePassProxy : public FCompositeCorePassProxy
 		{
 		public:
 			/** Constructor. */
-			FMergePassProxy(ECompositeCorePassInput Input0, ECompositeCorePassInput Input1);
+			FMergePassProxy(UE::CompositeCore::FPassInputArray InPassInputs);
 
-			FScreenPassTexture Add(FRDGBuilder& GraphBuilder, const FSceneView& InView, const FPostProcessMaterialInputs& Inputs, ISceneViewExtension::EPostProcessingPass InLocation) const override;
+			FScreenPassTexture Add(FRDGBuilder& GraphBuilder, const FSceneView& InView, const FPostProcessMaterialInputs& Inputs, const FIntRect& OutputViewRect, ISceneViewExtension::EPostProcessingPass InLocation) const override;
 
 			/** Merge operation. */
 			ECompositeCoreMergeOp MergeOp = ECompositeCoreMergeOp::Over;
diff --git a/Engine/Plugins/Experimental/Compositing/CompositeCore/Source/CompositeCore/Public/Passes/CompositeCorePassProxy.h b/Engine/Plugins/Experimental/Compositing/CompositeCore/Source/CompositeCore/Public/Passes/CompositeCorePassProxy.h
--- a/Engine/Plugins/Experimental/Compositing/CompositeCore/Source/CompositeCore/Public/Passes/CompositeCorePassProxy.h
+++ b/Engine/Plugins/Experimental/Compositing/CompositeCore/Source/CompositeCore/Public/Passes/CompositeCorePassProxy.h
@@ -5,32 +5,44 @@
 #include "CoreTypes.h"
 #include "Containers/ContainersFwd.h"
 #include "SceneViewExtension.h"
-#include "UObject/Object.h"
-
-#include "CompositeCorePassProxy.generated.h"
+#include "Misc/Optional.h"
+#include "Containers/Array.h"
 
 class FRDGBuilder;
+struct FRDGTextureDesc;
 class FSceneView;
 struct FScreenPassTexture;
+struct FScreenPassRenderTarget;
 struct FPostProcessMaterialInputs;
-class ICompositeCorePassProxy;
-
-/* Composite texture input type. */
-UENUM()
-enum class ECompositeCorePassInput : uint8
-{
-	/** Default input (usually scene color), or result of the previous pass. */
-	Default,
-	/** Texture from the built-in separate composite custom render pass of the specified primitives. */
-	CompositeRenderPass,
-	/** Direct composite texture, usually a media texture. */
-	CompositeTexture
-};
+class FCompositeCorePassProxy;
 
 namespace UE
 {
 	namespace CompositeCore
 	{
+		/* Type used to identify passes, textures or built-in renderer sources/targets. */
+		using ResourceId = uint32;
+		/* Built-in custom render pass identifier. */
+		static constexpr ResourceId BUILT_IN_CRP_ID = 1;
+		/* Built-in empty/black identifier. */
+		static constexpr ResourceId BUILT_IN_EMPTY_ID = 2;
+		/* First identifier of the external texture inputs, see ExternalInputs on FPostRenderWork. */
+		static constexpr ResourceId EXTERNAL_RANGE_START_ID = 100;
+
+		/** Pass texture description.	*/
+		struct FPassResourceDesc
+		{
+			ResourceId Id			= 0;
+			bool bInvertedAlpha		= false;
+			bool bDistorted			= false;
+		};
+
+		/* Optional pass input override, from the otherwise default 0-4 post-processing inputs. */
+		using FPassInput = TOptional<FPassResourceDesc>;
+		/* Array of pass inputs. */
+		using FPassInputArray = TArray<FPassInput>;
+
+		/** Options to control the built-in custom render pass. */
 		struct COMPOSITECORE_API FBuiltInRenderPassOptions final
 		{
 			/** Constructor. */
@@ -55,11 +67,11 @@ namespace UE
 			/** Constructor. */
 			FPostRenderWork();
 
-			/** Input (pre) processing passes. */
-			TArray<TPair<ECompositeCorePassInput, TArray<const ICompositeCorePassProxy*>>> InputProcessing;
+			/* Array of user-defined external input texture overrides, where indices map to ResourceId. */
+			TArray<const TWeakObjectPtr<UTexture>> ExternalInputs;
 
 			/** Post-processing passes at the specified locations. */
-			TSortedMap<ISceneViewExtension::EPostProcessingPass, TArray<const ICompositeCorePassProxy*>> FramePasses;
+			TSortedMap<ISceneViewExtension::EPostProcessingPass, TArray<const FCompositeCorePassProxy*>> FramePasses;
 
 			/** Proxy allocator. */
 			TUniquePtr<FSceneRenderingBulkObjectAllocator> FrameAllocator;
@@ -68,28 +80,58 @@ namespace UE
 }
 
 /** Render-thread pass proxy. */
-class COMPOSITECORE_API ICompositeCorePassProxy
+class COMPOSITECORE_API FCompositeCorePassProxy
 {
 public:
+
+	/** Default constructor. */
+	FCompositeCorePassProxy() = default;
+
 	/** Constructor. */
-	ICompositeCorePassProxy(ECompositeCorePassInput Input0)
-		: InputTypes({ Input0 })
-	{ };
+	FCompositeCorePassProxy(UE::CompositeCore::FPassInputArray InPassInputs)
+		: PassInputs(MoveTemp(InPassInputs))
+	{ }
 
-	/* Virtual destructor in case deletion is done on interface pointers. */
-	virtual ~ICompositeCorePassProxy() { };
+	/* Virtual destructor */
+	virtual ~FCompositeCorePassProxy() { };
 
 	/* Render-thread add pass method to override. */
-	virtual FScreenPassTexture Add(FRDGBuilder& GraphBuilder, const FSceneView& InView, const FPostProcessMaterialInputs& Inputs, ISceneViewExtension::EPostProcessingPass InLocation) const = 0;
+	virtual FScreenPassTexture Add(FRDGBuilder& GraphBuilder, const FSceneView& InView, const FPostProcessMaterialInputs& Inputs, const FIntRect& OutputViewRect, ISceneViewExtension::EPostProcessingPass InLocation) const = 0;
 
 	/* Number of inputs used by the pass. */
-	int32 GetNumInputs() const { return InputTypes.Num(); }
+	int32 GetNumInputs() const
+	{
+		return PassInputs.Num();
+	}
 
 	/* Number of inputs used by the pass. */
-	ECompositeCorePassInput GetPassInput(int32 Index) const { return (Index < GetNumInputs()) ? InputTypes[Index] : ECompositeCorePassInput::Default; };
+	const UE::CompositeCore::FPassInput& GetPassInput(int32 InputIndex) const
+	{
+		return PassInputs[InputIndex];
+	}
+
+	/* Return sub-passes per input index. */
+	const TArray<const FCompositeCorePassProxy*>* GetSubPasses(int32 InputIndex) const
+	{
+		return SubPasses.Find(InputIndex);
+	}
+
+	/* Add sub-passes per input index. */
+	void AddSubPasses(int32 InputIndex, TArray<const FCompositeCorePassProxy*> InSubPasses)
+	{
+		SubPasses.Emplace(InputIndex, MoveTemp(InSubPasses));
+	}
+
+	/* Convenience function to create an output render target with the specified resolution. */
+	static FScreenPassRenderTarget CreateOutputRenderTarget(FRDGBuilder& GraphBuilder,
+		const FSceneView& InView, const FIntRect& OutputViewRect, FRDGTextureDesc OutputDesc, const TCHAR* InName);
 
 protected:
+
 	/* List of pass input types. */
-	TArray<ECompositeCorePassInput> InputTypes;
+	UE::CompositeCore::FPassInputArray PassInputs;
+
+	/* Map of sub-passes per input index. */
+	TSortedMap<int32, TArray<const FCompositeCorePassProxy*>> SubPasses;
 };
 
