CompositeCore: Minor cleanup separating dilate pass from scene view extension.

diff --git a/Engine/Plugins/Experimental/Compositing/CompositeCore/Source/CompositeCore/Private/CompositeCoreSceneViewExtension.cpp b/Engine/Plugins/Experimental/Compositing/CompositeCore/Source/CompositeCore/Private/CompositeCoreSceneViewExtension.cpp
--- a/Engine/Plugins/Experimental/Compositing/CompositeCore/Source/CompositeCore/Private/CompositeCoreSceneViewExtension.cpp
+++ b/Engine/Plugins/Experimental/Compositing/CompositeCore/Source/CompositeCore/Private/CompositeCoreSceneViewExtension.cpp
@@ -3,6 +3,7 @@
 #include "CompositeCoreSceneViewExtension.h"
 
 #include "CompositeCoreModule.h"
+#include "Passes/CompositeCorePassDilate.h"
 
 #include "CommonRenderResources.h"
 #include "Components/PrimitiveComponent.h"
@@ -10,7 +11,6 @@
 #include "EngineUtils.h"
 #include "Engine/Texture.h"
 #include "HDRHelper.h"
-#include "PixelShaderUtils.h"
 #include "PostProcess/PostProcessMaterialInputs.h"
 #include "Rendering/CustomRenderPass.h"
 #include "RenderGraphBuilder.h"
@@ -20,80 +20,6 @@
 #include "SystemTextures.h"
 #include "TextureResource.h"
 
-DECLARE_GPU_STAT_NAMED(FCompositeCoreDilate, TEXT("CompositeCore.Dilate"));
-
-class FCompositeCoreDilateShader : public FGlobalShader
-{
-	DECLARE_GLOBAL_SHADER(FCompositeCoreDilateShader);
-	SHADER_USE_PARAMETER_STRUCT(FCompositeCoreDilateShader, FGlobalShader);
-
-	BEGIN_SHADER_PARAMETER_STRUCT(FParameters, )
-		SHADER_PARAMETER_RDG_TEXTURE(Texture2D, InputTexture)
-		SHADER_PARAMETER_RDG_TEXTURE_UAV(RWTexture2D<float4>, RWOutputTexture)
-		SHADER_PARAMETER(FIntPoint, Dimensions)
-	END_SHADER_PARAMETER_STRUCT()
-
-	static inline void ModifyCompilationEnvironment(const FGlobalShaderPermutationParameters& Parameters, FShaderCompilerEnvironment& OutEnvironment)
-	{
-		OutEnvironment.SetDefine(TEXT("TILE_SIZE"), ThreadGroupSize);
-
-		FGlobalShader::ModifyCompilationEnvironment(Parameters, OutEnvironment);
-	}
-
-	static const uint32 ThreadGroupSize = 16;
-};
-IMPLEMENT_GLOBAL_SHADER(FCompositeCoreDilateShader, "/Plugin/CompositeCore/Private/CompositeCoreDilate.usf", "MainCS", SF_Compute);
-
-namespace UE
-{
-	namespace CompositeCore
-	{
-		namespace Private
-		{
-			void AddDilatePass(FRDGBuilder& GraphBuilder, FRDGTextureRef Input, FRDGTextureRef Output, ERHIFeatureLevel::Type FeatureLevel)
-			{
-				RDG_EVENT_SCOPE_STAT(GraphBuilder, FCompositeCoreDilate, "CompositeCore.Dilate");
-				RDG_GPU_STAT_SCOPE(GraphBuilder, FCompositeCoreDilate);
-
-				FGlobalShaderMap* GlobalShaderMap = GetGlobalShaderMap(FeatureLevel);
-				const FIntPoint TextureSize = Input->Desc.Extent;
-
-				FCompositeCoreDilateShader::FParameters* PassParameters = GraphBuilder.AllocParameters<FCompositeCoreDilateShader::FParameters>();
-				PassParameters->InputTexture = Input;
-				PassParameters->RWOutputTexture = GraphBuilder.CreateUAV(Output);
-				PassParameters->Dimensions = TextureSize;
-
-				TShaderMapRef<FCompositeCoreDilateShader> ComputeShader(GlobalShaderMap);
-				FComputeShaderUtils::AddPass(
-					GraphBuilder,
-					RDG_EVENT_NAME("CompositeCore.Dilate (%dx%d)", TextureSize.X, TextureSize.Y),
-					GSupportsEfficientAsyncCompute ? ERDGPassFlags::AsyncCompute : ERDGPassFlags::Compute,
-					ComputeShader,
-					PassParameters,
-					FComputeShaderUtils::GetGroupCount(TextureSize, FCompositeCoreDilateShader::ThreadGroupSize)
-				);
-			}
-
-			FRDGTextureDesc GetPostProcessingDesc(FRDGTextureDesc InDesc, EPixelFormat InOutputFormat)
-			{
-				FRDGTextureDesc OutputDesc = InDesc;
-				OutputDesc.Dimension = ETextureDimension::Texture2D;
-				OutputDesc.ArraySize = 1;
-
-				OutputDesc.Reset();
-				if (InOutputFormat != PF_Unknown)
-				{
-					OutputDesc.Format = InOutputFormat;
-				}
-				OutputDesc.ClearValue = FClearValueBinding(FLinearColor::Black);
-				OutputDesc.Flags &= (~ETextureCreateFlags::FastVRAM);
-
-				return OutputDesc;
-			}
-		}
-	}
-}
-
 class FCompositeCoreCustomRenderPass : public FCustomRenderPassBase
 {
 public:
diff --git a/Engine/Plugins/Experimental/Compositing/CompositeCore/Source/CompositeCore/Private/Passes/CompositeCorePassDilate.cpp b/Engine/Plugins/Experimental/Compositing/CompositeCore/Source/CompositeCore/Private/Passes/CompositeCorePassDilate.cpp
new file mode 100644
--- /dev/null
+++ b/Engine/Plugins/Experimental/Compositing/CompositeCore/Source/CompositeCore/Private/Passes/CompositeCorePassDilate.cpp
@@ -0,0 +1,83 @@
+// Copyright Epic Games, Inc. All Rights Reserved.
+
+#include "CompositeCorePassDilate.h"
+
+#include "RHIFeatureLevel.h"
+#include "RenderGraphBuilder.h"
+#include "RenderGraphUtils.h"
+#include "PixelShaderUtils.h"
+
+DECLARE_GPU_STAT_NAMED(FCompositeCoreDilate, TEXT("CompositeCore.Dilate"));
+
+class FCompositeCoreDilateShader : public FGlobalShader
+{
+	DECLARE_GLOBAL_SHADER(FCompositeCoreDilateShader);
+	SHADER_USE_PARAMETER_STRUCT(FCompositeCoreDilateShader, FGlobalShader);
+
+	BEGIN_SHADER_PARAMETER_STRUCT(FParameters, )
+		SHADER_PARAMETER_RDG_TEXTURE(Texture2D, InputTexture)
+		SHADER_PARAMETER_RDG_TEXTURE_UAV(RWTexture2D<float4>, RWOutputTexture)
+		SHADER_PARAMETER(FIntPoint, Dimensions)
+	END_SHADER_PARAMETER_STRUCT()
+
+	static inline void ModifyCompilationEnvironment(const FGlobalShaderPermutationParameters& Parameters, FShaderCompilerEnvironment& OutEnvironment)
+	{
+		OutEnvironment.SetDefine(TEXT("TILE_SIZE"), ThreadGroupSize);
+
+		FGlobalShader::ModifyCompilationEnvironment(Parameters, OutEnvironment);
+	}
+
+	static const uint32 ThreadGroupSize = 16;
+};
+IMPLEMENT_GLOBAL_SHADER(FCompositeCoreDilateShader, "/Plugin/CompositeCore/Private/CompositeCoreDilate.usf", "MainCS", SF_Compute);
+
+namespace UE
+{
+	namespace CompositeCore
+	{
+		namespace Private
+		{
+			void AddDilatePass(FRDGBuilder& GraphBuilder, FRDGTextureRef Input, FRDGTextureRef Output, ERHIFeatureLevel::Type FeatureLevel)
+			{
+				RDG_EVENT_SCOPE_STAT(GraphBuilder, FCompositeCoreDilate, "CompositeCore.Dilate");
+				RDG_GPU_STAT_SCOPE(GraphBuilder, FCompositeCoreDilate);
+
+				FGlobalShaderMap* GlobalShaderMap = GetGlobalShaderMap(FeatureLevel);
+				const FIntPoint TextureSize = Input->Desc.Extent;
+
+				FCompositeCoreDilateShader::FParameters* PassParameters = GraphBuilder.AllocParameters<FCompositeCoreDilateShader::FParameters>();
+				PassParameters->InputTexture = Input;
+				PassParameters->RWOutputTexture = GraphBuilder.CreateUAV(Output);
+				PassParameters->Dimensions = TextureSize;
+
+				TShaderMapRef<FCompositeCoreDilateShader> ComputeShader(GlobalShaderMap);
+				FComputeShaderUtils::AddPass(
+					GraphBuilder,
+					RDG_EVENT_NAME("CompositeCore.Dilate (%dx%d)", TextureSize.X, TextureSize.Y),
+					GSupportsEfficientAsyncCompute ? ERDGPassFlags::AsyncCompute : ERDGPassFlags::Compute,
+					ComputeShader,
+					PassParameters,
+					FComputeShaderUtils::GetGroupCount(TextureSize, FCompositeCoreDilateShader::ThreadGroupSize)
+				);
+			}
+
+			FRDGTextureDesc GetPostProcessingDesc(FRDGTextureDesc InDesc, EPixelFormat InOutputFormat)
+			{
+				FRDGTextureDesc OutputDesc = InDesc;
+				OutputDesc.Dimension = ETextureDimension::Texture2D;
+				OutputDesc.ArraySize = 1;
+
+				OutputDesc.Reset();
+				if (InOutputFormat != PF_Unknown)
+				{
+					OutputDesc.Format = InOutputFormat;
+				}
+				OutputDesc.ClearValue = FClearValueBinding(FLinearColor::Black);
+				OutputDesc.Flags &= (~ETextureCreateFlags::FastVRAM);
+
+				return OutputDesc;
+			}
+		}
+	}
+}
+
diff --git a/Engine/Plugins/Experimental/Compositing/CompositeCore/Source/CompositeCore/Private/Passes/CompositeCorePassDilate.h b/Engine/Plugins/Experimental/Compositing/CompositeCore/Source/CompositeCore/Private/Passes/CompositeCorePassDilate.h
new file mode 100644
--- /dev/null
+++ b/Engine/Plugins/Experimental/Compositing/CompositeCore/Source/CompositeCore/Private/Passes/CompositeCorePassDilate.h
@@ -0,0 +1,15 @@
+// Copyright Epic Games, Inc. All Rights Reserved.
+
+#include "RenderGraphFwd.h"
+#include "RHIFwd.h"
+
+namespace UE
+{
+	namespace CompositeCore
+	{
+		namespace Private
+		{
+			void AddDilatePass(FRDGBuilder& GraphBuilder, FRDGTextureRef Input, FRDGTextureRef Output, ERHIFeatureLevel::Type FeatureLevel);
+		}
+	}
+}
