Updated alpha/holdout project setting toast notifications & warnings.

Increased clarity for settings updating DefaultEngine.ini, and potential additional renderer costs.

diff --git a/Engine/Plugins/Experimental/Avalanche/Source/AvalancheMediaEditor/Private/Rundown/Preview/SAvaRundownPagePreview.cpp b/Engine/Plugins/Experimental/Avalanche/Source/AvalancheMediaEditor/Private/Rundown/Preview/SAvaRundownPagePreview.cpp
--- a/Engine/Plugins/Experimental/Avalanche/Source/AvalancheMediaEditor/Private/Rundown/Preview/SAvaRundownPagePreview.cpp
+++ b/Engine/Plugins/Experimental/Avalanche/Source/AvalancheMediaEditor/Private/Rundown/Preview/SAvaRundownPagePreview.cpp
@@ -307,8 +307,9 @@ void SAvaRundownPagePreview::HandleCheckerboardActionExecute() const
 		if (PropagateAlphaCVar && PropagateAlphaCVar->GetBool() != true)
 		{
 			const FText NotificationText = LOCTEXT("AlphaSupport",
-				"An output requested Alpha Support but the required project setting is not enabled!\n"
-				"Go to Project Settings > Rendering > PostProcessing > 'Alpha Output' and enable it.");
+				"An output requested alpha but the required project setting is not enabled!\n"
+				"Warning: update can add renderer performance costs.\n"
+				"Go to Project Settings > Rendering > Default Settings > 'Alpha Output' to enable it in DefaultEngine.ini.");
 
 			FNotificationInfo Info(NotificationText);
 			Info.ExpireDuration = 5.0f;	// The message is long, need more time to read it.
diff --git a/Engine/Plugins/Experimental/Avalanche/Source/AvalancheViewport/Private/Visualizers/AvaViewportPostProcessVisualizer.cpp b/Engine/Plugins/Experimental/Avalanche/Source/AvalancheViewport/Private/Visualizers/AvaViewportPostProcessVisualizer.cpp
--- a/Engine/Plugins/Experimental/Avalanche/Source/AvalancheViewport/Private/Visualizers/AvaViewportPostProcessVisualizer.cpp
+++ b/Engine/Plugins/Experimental/Avalanche/Source/AvalancheViewport/Private/Visualizers/AvaViewportPostProcessVisualizer.cpp
@@ -76,7 +76,7 @@ bool FAvaViewportPostProcessVisualizer::CanActivate(bool bInSilent) const
 
 	const EAppReturnType::Type Response = FMessageDialog::Open(
 		EAppMsgType::YesNoCancel,
-		LOCTEXT("AlphaChannelInPostProcessingRequiredMessage", "This Post Process effect will not work without enabling Alpha Support in the Tonemapper settings via Project Settings > Engine > Rendering > Post Processing > Enable Alpha Channel Support. It must be set to Allow Though Tonemapper.\n\nEnable this setting now? (Restart and shader recompilation required.)"),
+		LOCTEXT("AlphaChannelInPostProcessingRequiredMessage", "This Post Process effect will not work without enabling Alpha Output via Project Settings > Engine > Rendering > Default Settings > 'Alpha Ouput'. Warning: update can add renderer performance costs.\n\nEnable this setting in DefaultEngine.ini now?"),
 		LOCTEXT("AlphaChannelInPostProcessingRequiredTitle", "Project Setting Required")
 	);
 
diff --git a/Engine/Plugins/Experimental/Compositing/CompositeCore/Source/CompositeCore/Private/CompositeCoreSubsystem.cpp b/Engine/Plugins/Experimental/Compositing/CompositeCore/Source/CompositeCore/Private/CompositeCoreSubsystem.cpp
--- a/Engine/Plugins/Experimental/Compositing/CompositeCore/Source/CompositeCore/Private/CompositeCoreSubsystem.cpp
+++ b/Engine/Plugins/Experimental/Compositing/CompositeCore/Source/CompositeCore/Private/CompositeCoreSubsystem.cpp
@@ -176,8 +176,8 @@ void UCompositeCoreSubsystem::PrimitiveHoldoutSettingsNotification(URendererSett
 
 	const FText AlphaOutputSettingOption = LOCTEXT("HoldoutSetting_AlphaOutput", "\n- Alpha Output");
 	const FText PrimitiveHoldoutSettingOption = LOCTEXT("HoldoutSetting_PrimitiveHoldout", "\n- Support Primitive Alpha Holdout");
-	const FText HoldoutText = FText::Format(LOCTEXT("HoldoutSettingPrompt", "The following project setting(s) must be enabled for holdout composite:{0}{1}"), bAlphaOutputMissing ? AlphaOutputSettingOption : FText::GetEmpty(), bPrimitiveHoldoutMissing ? PrimitiveHoldoutSettingOption : FText::GetEmpty());
-	const FText HoldoutConfirmText = LOCTEXT("HoldoutSettingConfirm", "Enable");
+	const FText HoldoutText = FText::Format(LOCTEXT("HoldoutSettingPrompt", "The following project setting(s) must be enabled for holdout composite:{0}{1}\n\nWarning: update can add renderer performance costs."), bAlphaOutputMissing ? AlphaOutputSettingOption : FText::GetEmpty(), bPrimitiveHoldoutMissing ? PrimitiveHoldoutSettingOption : FText::GetEmpty());
+	const FText HoldoutConfirmText = LOCTEXT("HoldoutSettingConfirm", "Enable (DefaultEngine.ini)");
 	const FText HoldoutCancelText = LOCTEXT("HoldoutSettingCancel", "Not Now");
 
 	/** Utility functions for notifications */
diff --git a/Engine/Plugins/MovieScene/MovieRenderPipeline/Source/MovieRenderPipelineCore/Private/Graph/MovieGraphUtils.cpp b/Engine/Plugins/MovieScene/MovieRenderPipeline/Source/MovieRenderPipelineCore/Private/Graph/MovieGraphUtils.cpp
--- a/Engine/Plugins/MovieScene/MovieRenderPipeline/Source/MovieRenderPipelineCore/Private/Graph/MovieGraphUtils.cpp
+++ b/Engine/Plugins/MovieScene/MovieRenderPipeline/Source/MovieRenderPipelineCore/Private/Graph/MovieGraphUtils.cpp
@@ -115,8 +115,8 @@ namespace UE::MovieGraph
 
 			const FText AlphaOutputSettingOption = LOCTEXT("MovieAlphaSetting_AlphaOutput", "\n- Alpha Output");
 			const FText PrimitiveHoldoutSettingOption = LOCTEXT("MovieAlphaSetting_PrimitiveHoldout", "\n- Support Primitive Alpha Holdout");
-			const FText MovieAlphaText = FText::Format(LOCTEXT("MovieAlphaSettingPrompt", "The following project setting(s) must be enabled for Movie Render Graph's {0}:{1}{2}"), InRequestingFeatureName, bAlphaOutputMissing ? AlphaOutputSettingOption : FText::GetEmpty(), bPrimitiveHoldoutMissing ? PrimitiveHoldoutSettingOption : FText::GetEmpty());
-			const FText MovieAlphaConfirmText = LOCTEXT("MovieAlphaSettingConfirm", "Enable");
+			const FText MovieAlphaText = FText::Format(LOCTEXT("MovieAlphaSettingPrompt", "The 'Support Primitive Alpha Holdout' project setting w/ 'Alpha Ouput' must be enabled for Movie Render Graph's {0}.\n\nWarning: update can add renderer performance costs."), InRequestingFeatureName);
+			const FText MovieAlphaConfirmText = LOCTEXT("MovieAlphaSettingConfirm", "Enable (DefaultEngine.ini)");
 			const FText MovieAlphaCancelText = LOCTEXT("MovieAlphaSettingCancel", "Not Now");
 
 			/** Utility functions for notifications */
