Composite: First pass at full layer implementation.

CompositeCore:
* Composite passes have explicit pass struct for declaring & receving pass inputs, context and output. These structs define accompanying metadata such that operations can be done in consistent distortion & color encoding.
* Added "None" merge pass operation type.

Composite:
* Composite actor supports ordered layers acting as merge operations, with optional subpasses.
* Shadow reflection capture is supported and (custom) scene capture components are automatically managed on the actor.
* Bug fix for toggling the "Enabled" state on the actor, so that primitives are already unregistered.

diff --git a/Engine/Plugins/Experimental/Compositing/CompositeCore/Shaders/Private/CompositeCoreMerge.usf b/Engine/Plugins/Experimental/Compositing/CompositeCore/Shaders/Private/CompositeCoreMerge.usf
--- a/Engine/Plugins/Experimental/Compositing/CompositeCore/Shaders/Private/CompositeCoreMerge.usf
+++ b/Engine/Plugins/Experimental/Compositing/CompositeCore/Shaders/Private/CompositeCoreMerge.usf
@@ -13,10 +13,11 @@
 #define EDistortionUV_Distorted   1
 #define EDistortionUV_Undistorted 2
 
-#define ECompositeCoreMergeOp_Over 0
-#define ECompositeCoreMergeOp_Add  1
-#define ECompositeCoreMergeOp_Mul  2
-#define ECompositeCoreMergeOp_Div  3
+#define ECompositeCoreMergeOp_None 0
+#define ECompositeCoreMergeOp_Over 1
+#define ECompositeCoreMergeOp_Add  2
+#define ECompositeCoreMergeOp_Mul  3
+#define ECompositeCoreMergeOp_Div  4
 
 SCREEN_PASS_TEXTURE_VIEWPORT(Input0)
 SCREEN_PASS_TEXTURE_VIEWPORT(Input1)
@@ -43,7 +44,7 @@ SamplerState UndistortingDisplacementSampler;
 float2 DisplayGamma; // DisplayGamma, 1/DisplayGamma
 uint MergeOp;
 uint OutputEncoding;
-uint bFlags; // 0: bOutputSceneColor, 1: bUnpremulInput1
+uint bFlags; // 0: bInvertedAlphaOutput, 1: bUnpremulInput1
 
 float SafeRcp(float x)
 {
@@ -181,7 +182,6 @@ void MainPS(
     // Re-encode color
     OutColor.rgb = Encode(OutColor.rgb, OutputEncoding);
 
-    // Optionally re-invert alpha
-    const bool bOutputSceneColor = (bFlags & 1);
-    OutColor.a = saturate(select(bOutputSceneColor, 1.0 - OutColor.a, OutColor.a));
+	const bool bInvertedAlphaOutput = (bFlags & 1);
+	OutColor.a = saturate(select(bInvertedAlphaOutput, 1.0 - OutColor.a, OutColor.a));
 }
diff --git a/Engine/Plugins/Experimental/Compositing/CompositeCore/Source/CompositeCore/Private/CompositeCoreSceneViewExtension.cpp b/Engine/Plugins/Experimental/Compositing/CompositeCore/Source/CompositeCore/Private/CompositeCoreSceneViewExtension.cpp
--- a/Engine/Plugins/Experimental/Compositing/CompositeCore/Source/CompositeCore/Private/CompositeCoreSceneViewExtension.cpp
+++ b/Engine/Plugins/Experimental/Compositing/CompositeCore/Source/CompositeCore/Private/CompositeCoreSceneViewExtension.cpp
@@ -11,6 +11,7 @@
 #include "EngineUtils.h"
 #include "Engine/Texture.h"
 #include "HDRHelper.h"
+#include "PostProcess/LensDistortion.h"
 #include "PostProcess/PostProcessMaterialInputs.h"
 #include "Rendering/CustomRenderPass.h"
 #include "RenderGraphBuilder.h"
@@ -126,14 +127,18 @@ void FCompositeCoreSceneViewExtension::SetPostRenderWork_GameThread(UE::Composit
 
 				for (int32 Index=0; Index < InWork.ExternalInputs.Num(); ++Index)
 				{
-					TStrongObjectPtr<UTexture> ExternalInput = InWork.ExternalInputs[Index].Pin();
+					const UE::CompositeCore::FExternalTexture& ExternalInput = InWork.ExternalInputs[Index];
+					TStrongObjectPtr<UTexture> ExternalTexture = ExternalInput.Texture.Pin();
 					
-					if (ExternalInput.IsValid())
+					if (ExternalTexture.IsValid())
 					{
-						if (FTextureResource* TexResource = ExternalInput->GetResource())
+						if (FTextureResource* TexResource = ExternalTexture->GetResource())
 						{
-							const int32 ExternalId = UE::CompositeCore::EXTERNAL_RANGE_START_ID + Index;
-							SVE->ExternalInputs_RenderThread.Add(ExternalId, CreateRenderTarget(TexResource->GetTextureRHI(), TEXT("CompositeExternalInput")));
+							const int32 MapIndex = UE::CompositeCore::EXTERNAL_RANGE_START_ID + Index;
+
+							UE::CompositeCore::FExternalRenderTarget& Target = SVE->ExternalInputs_RenderThread.Add(MapIndex);
+							Target.RenderTarget = CreateRenderTarget(TexResource->GetTextureRHI(), TEXT("CompositeExternalInput"));
+							Target.Metadata = ExternalInput.Metadata;
 						}
 					}
 				}
@@ -344,98 +349,167 @@ bool FCompositeCoreSceneViewExtension::IsActiveForView(const FSceneView& InView)
 	return bIsActive;
 }
 
-TSortedMap<int32, FScreenPassTextureSlice> FCompositeCoreSceneViewExtension::CreateInputOverridesMap(FRDGBuilder& GraphBuilder, const FSceneView& InView) const
+TSortedMap<int32, UE::CompositeCore::FPassInput> FCompositeCoreSceneViewExtension::CreateExternalTextureMap(FRDGBuilder& GraphBuilder, const FSceneView& InView, const FPostProcessMaterialInputs& Inputs) const
 {
 	using namespace UE::CompositeCore;
 
-	TSortedMap<int32, FScreenPassTextureSlice> InputOverrides;
+	TSortedMap<int32, FPassInput> InputOverrides;
 	InputOverrides.Reserve(2 + ExternalInputs_RenderThread.Num());
 
-	{
-		FScreenPassTexture ScreenTex = FScreenPassTexture{ GSystemTextures.GetBlackDummy(GraphBuilder) };
-		InputOverrides.Emplace(BUILT_IN_EMPTY_ID, FScreenPassTextureSlice::CreateFromScreenPassTexture(GraphBuilder, ScreenTex));
-	}
+	InputOverrides.Emplace(BUILT_IN_EMPTY_ID, FScreenPassTexture{ GSystemTextures.GetBlackDummy(GraphBuilder) });
 
 	const TRefCountPtr<IPooledRenderTarget>* CompositeRenderPassPtr = CustomRenderTargetPerView_RenderThread.Find(InView.GetViewKey());
 	if (CompositeRenderPassPtr != nullptr)
 	{
-		const FScreenPassTexture ScreenTex = FScreenPassTexture{ GraphBuilder.RegisterExternalTexture(*CompositeRenderPassPtr) };
-		InputOverrides.Emplace(BUILT_IN_CRP_ID, FScreenPassTextureSlice::CreateFromScreenPassTexture(GraphBuilder, ScreenTex));
+		FPassInput& Resource = InputOverrides.Add(BUILT_IN_CRP_ID);
+		Resource.Texture = FScreenPassTexture{ GraphBuilder.RegisterExternalTexture(*CompositeRenderPassPtr) };
+		Resource.Metadata.bInvertedAlpha = true;
 	}
 
 	for (const auto& Pair : ExternalInputs_RenderThread)
 	{
-		const FScreenPassTexture ScreenTex = FScreenPassTexture{ GraphBuilder.RegisterExternalTexture(Pair.Value) };
-		InputOverrides.Emplace(Pair.Key, FScreenPassTextureSlice::CreateFromScreenPassTexture(GraphBuilder, ScreenTex));
+		FPassInput& Resource = InputOverrides.Add(Pair.Key);
+		Resource.Texture = FScreenPassTexture{ GraphBuilder.RegisterExternalTexture(Pair.Value.RenderTarget) };
+		Resource.Metadata = Pair.Value.Metadata;
 	}
 
 	return InputOverrides;
 }
 
-bool FCompositeCoreSceneViewExtension::ApplyPasses_Recursive(FRDGBuilder& GraphBuilder, const FSceneView& InView,
-	const FPostProcessMaterialInputs& Inputs, ISceneViewExtension::EPostProcessingPass InLocation,
-	const TArray<const FCompositeCorePassProxy*> InPasses, const TSortedMap<int32, FScreenPassTextureSlice>& InputOverrides,
-	const FIntRect& SceneColorRect, int32 RecursionLevel, FScreenPassTexture& Output)
+void FCompositeCoreSceneViewExtension::UpdateNextPassInputs(
+	FRDGBuilder& GraphBuilder,
+	const UE::CompositeCore::FPassOutput& InOutput,
+	const int32 BindingIndex,
+	UE::CompositeCore::FPassInputArray& InOutInputs,
+	FExternalTextureMap& InOutExternalTextures
+)
 {
+	using namespace UE::CompositeCore;
+
+	if (InOutput.Override.IsSet())
+	{
+		const ResourceId& ExternalId = InOutput.Override.GetValue();
+
+		FPassInput* ExternalTexture = InOutExternalTextures.Find(ExternalId);
+		if (ensureMsgf(ExternalTexture != nullptr, TEXT("Unexpected missing external texture override as output.")))
+		{
+			*ExternalTexture = InOutput.Resource;
+		}
+	}
+	else
+	{
+		// Update input binding for the next pass(es)
+		InOutInputs[BindingIndex] = InOutput.Resource;
+	}
+}
+
+bool FCompositeCoreSceneViewExtension::ApplyPasses_Recursive(
+	FRDGBuilder& GraphBuilder,
+	const FSceneView& InView,
+	const UE::CompositeCore::FPassInputArray& Inputs,
+	const UE::CompositeCore::FPassInputArray& OriginalInputs,
+	UE::CompositeCore::FPassContext& PassContext,
+	const TArray<const FCompositeCorePassProxy*> InPasses,
+	FExternalTextureMap& ExternalTextures,
+	int32 RecursionLevel,
+	UE::CompositeCore::FPassOutput& Output)
+{
+	using namespace UE::CompositeCore;
+
 	if (InPasses.IsEmpty())
 	{
 		return false;
 	}
 
 	// Default pass inputs
-	FPostProcessMaterialInputs BasePassInputs = Inputs;
+	FPassInputArray BasePassInputs = Inputs;
 	
 	// Iterate over all passes
-	for (int32 Index = 0; Index < InPasses.Num(); ++Index)
+	for (int32 PassIndex = 0; PassIndex < InPasses.Num(); ++PassIndex)
 	{
-		const FCompositeCorePassProxy* Pass = InPasses[Index];
+		const FCompositeCorePassProxy* Pass = InPasses[PassIndex];
 		
 		// Output override only applies to the (base) last pass
-		const bool bIsLastPass = (Index == InPasses.Num() - 1) && (RecursionLevel == 0);
+		const bool bIsLastPass = (PassIndex == InPasses.Num() - 1) && (RecursionLevel == 0);
 
 		// Update inputs for the current pass
-		FPostProcessMaterialInputs PassInputs = BasePassInputs;
+		FPassInputArray PassInputs = BasePassInputs;
 
-		for (int32 InputIndex = 0; InputIndex < Pass->GetNumInputs(); ++InputIndex)
+		// Iterate over declared pass inputs
+		for (int32 InputIndex = 0; InputIndex < Pass->GetNumDeclaredInputs(); ++InputIndex)
 		{
-			if (!ensureMsgf(InputIndex <= static_cast<int32>(EPostProcessMaterialInput::Velocity), TEXT("Inputs exceeding the number of post-process material inputs is currently not allowed.")))
+			const TArray<const FCompositeCorePassProxy*>* SubPasses = Pass->GetSubPasses(InputIndex);
+			if (SubPasses)
 			{
-				break;
+				// Recursively apply sub-passes per input index, automatically overriding inputs.
+				UE::CompositeCore::FPassOutput SubPassOutput;
+				if (ApplyPasses_Recursive(GraphBuilder, InView, PassInputs, OriginalInputs, PassContext, *SubPasses, ExternalTextures, RecursionLevel + 1, SubPassOutput))
+				{
+					UpdateNextPassInputs(GraphBuilder, SubPassOutput, InputIndex, PassInputs, ExternalTextures);
+				}
 			}
-
-			// Apply input overrides
-			const UE::CompositeCore::FPassInput& PassInput = Pass->GetPassInput(InputIndex);
-			if (PassInput.IsSet())
+			else
 			{
-				if (const FScreenPassTextureSlice* SliceInput = InputOverrides.Find(PassInput.GetValue().Id))
+				// No sub-pass input override, fetch the regular input index
+				const FPassInputDecl& DeclaredInput = Pass->GetDeclaredInput(InputIndex);
+
+				// If an internal texture resource is expected, connect the current input index with either the original or current bass pass input.
+				if (DeclaredInput.IsType<FPassInternalResourceDesc>())
 				{
-					PassInputs.SetInput(static_cast<EPostProcessMaterialInput>(InputIndex), *SliceInput);
+					const FPassInternalResourceDesc& Desc = DeclaredInput.Get<FPassInternalResourceDesc>();
+					const int32 DeclaredInputIndex = Desc.Index;
+
+					if (Desc.bOriginalCopyBeforePasses)
+					{
+						if (ensureMsgf(OriginalInputs.IsValidIndex(DeclaredInputIndex), TEXT("Invalid internal input: %d"), DeclaredInputIndex))
+						{
+							PassInputs[InputIndex] = OriginalInputs[DeclaredInputIndex];
+						}
+					}
+					else
+					{
+						if (ensureMsgf(BasePassInputs.IsValidIndex(DeclaredInputIndex), TEXT("Invalid internal input: %d"), DeclaredInputIndex))
+						{
+							PassInputs[InputIndex] = BasePassInputs[DeclaredInputIndex];
+						}
+					}
 				}
-			}
+				// If an external texture resource is expected, connect the current input index with the resolved identifier
+				else if (DeclaredInput.IsType<FPassExternalResourceDesc>())
+				{
+					const FPassExternalResourceDesc& Desc = DeclaredInput.Get<FPassExternalResourceDesc>();
+					const ResourceId DeclaredExternalId = Desc.Id;
 
-			// Recursively apply sub-passes per input index, automatically overriding inputs.
-			const TArray<const FCompositeCorePassProxy*>* SubPasses = Pass->GetSubPasses(InputIndex);
-			if (SubPasses)
-			{
-				FScreenPassTexture SubPassOutput;
-				if (ApplyPasses_Recursive(GraphBuilder, InView, Inputs, InLocation, *SubPasses, InputOverrides, SceneColorRect, RecursionLevel + 1, SubPassOutput))
+					const FPassInput* SliceInput = ExternalTextures.Find(DeclaredExternalId);
+					if (ensureMsgf(SliceInput, TEXT("Invalid external input: %d"), DeclaredExternalId))
+					{
+						PassInputs[InputIndex] = *SliceInput;
+					}
+				}
+				else
 				{
-					PassInputs.SetInput(static_cast<EPostProcessMaterialInput>(InputIndex), FScreenPassTextureSlice::CreateFromScreenPassTexture(GraphBuilder, SubPassOutput));
+					checkNoEntry();
 				}
 			}
 		}
 
-		// Only apply the output override on the last pass
-		if (!bIsLastPass)
+		if (bIsLastPass)
 		{
+			// Invert alpha when writing back to scene color
+			PassContext.bOutputSceneColor = true;
+			// Optionally apply distortion
+			PassContext.bOutputDistortionEnabled = PassContext.Location > ISceneViewExtension::EPostProcessingPass::SSRInput;
+		}
+		else
+		{
+			// Only apply the output override on the last pass
 			PassInputs.OverrideOutput = FScreenPassRenderTarget{};
 		}
 
 		// Register pass and update output
-		Output = Pass->Add(GraphBuilder, InView, PassInputs, SceneColorRect, InLocation);
+		Output = Pass->Add(GraphBuilder, InView, PassInputs, PassContext);
 
-		// Update input 0 for the next pass
-		BasePassInputs.SetInput(EPostProcessMaterialInput::SceneColor, FScreenPassTextureSlice::CreateFromScreenPassTexture(GraphBuilder, Output));
+		UpdateNextPassInputs(GraphBuilder, Output, 0, BasePassInputs, ExternalTextures);
 	}
 
 	return true;
@@ -445,18 +519,28 @@ FScreenPassTexture FCompositeCoreSceneViewExtension::PostProcessWork_RenderThrea
 {
 	using namespace UE::CompositeCore;
 
-	TSortedMap<int32, FScreenPassTextureSlice> InputOverrides = CreateInputOverridesMap(GraphBuilder, InView);
-	const FIntRect& SceneColorRect = Inputs.GetInput(EPostProcessMaterialInput::SceneColor).ViewRect;
+	FExternalTextureMap ExternalTextures = CreateExternalTextureMap(GraphBuilder, InView, Inputs);
 	const TArray<const FCompositeCorePassProxy*>& Passes = GetRenderWork().FramePasses[InLocation];
 
-	FScreenPassTexture Output;
+	FPassContext PassContext;
+	PassContext.SceneTextures = Inputs.SceneTextures;
+	PassContext.OutputViewRect = Inputs.GetInput(EPostProcessMaterialInput::SceneColor).ViewRect;
+	PassContext.Location = InLocation;
+	PassContext.bOutputSceneColor = false;
+	PassContext.bOutputDistortionEnabled = false;
+
+	const FLensDistortionLUT& LensDistortionLUT = LensDistortion::GetLUTUnsafe(InView);
+	const bool bLensDistortionEnabled = LensDistortionLUT.IsEnabled() && (LensDistortion::GetPassLocationUnsafe(InView) == LensDistortion::EPassLocation::TSR);
+
+	FPassInputArray ResolvedInputs(GraphBuilder, InView, Inputs, InLocation, bLensDistortionEnabled);
+	FPassOutput Output;
 	
 	constexpr int32 RecursionLevel = 0;
-	ApplyPasses_Recursive(GraphBuilder, InView, Inputs, InLocation, Passes, InputOverrides, SceneColorRect, RecursionLevel, Output);
+	ApplyPasses_Recursive(GraphBuilder, InView, ResolvedInputs, ResolvedInputs, PassContext, Passes, ExternalTextures, RecursionLevel, Output);
 
-	if (Output.IsValid())
+	if (Output.Resource.Texture.IsValid() && !Output.Override.IsSet())
 	{
-		return MoveTemp(Output);
+		return MoveTemp(Output.Resource.Texture);
 	}
 	else
 	{
diff --git a/Engine/Plugins/Experimental/Compositing/CompositeCore/Source/CompositeCore/Private/CompositeCoreSceneViewExtension.h b/Engine/Plugins/Experimental/Compositing/CompositeCore/Source/CompositeCore/Private/CompositeCoreSceneViewExtension.h
--- a/Engine/Plugins/Experimental/Compositing/CompositeCore/Source/CompositeCore/Private/CompositeCoreSceneViewExtension.h
+++ b/Engine/Plugins/Experimental/Compositing/CompositeCore/Source/CompositeCore/Private/CompositeCoreSceneViewExtension.h
@@ -18,6 +18,21 @@ struct FScreenPassTextureSlice;
 class UMaterialInterface;
 class UTexture;
 
+enum class EPostProcessMaterialInput : uint32;
+
+namespace UE
+{
+	namespace CompositeCore
+	{
+		/** External pooled render target and its accompanying metadata. */
+		struct FExternalRenderTarget
+		{
+			TRefCountPtr<IPooledRenderTarget> RenderTarget = {};
+			FResourceMetadata Metadata = {};
+		};
+	}
+}
+
 class FCompositeCoreSceneViewExtension : public FWorldSceneViewExtension
 {
 public:
@@ -61,18 +76,38 @@ protected:
 	//~ End ISceneViewExtension Interface
 
 private:
+	using FExternalTextureMap = TSortedMap<int32, UE::CompositeCore::FPassInput>;
+
 	/* Getter for the frame render work, either from the actor or default. */
 	const UE::CompositeCore::FPostRenderWork& GetRenderWork() const;
 
 	/* Active check to avoid work in post-processing. */
 	bool IsActiveForView(const FSceneView& InView) const;
 
-	TSortedMap<int32, FScreenPassTextureSlice> CreateInputOverridesMap(FRDGBuilder& GraphBuilder, const FSceneView& InView) const;
-
-	bool ApplyPasses_Recursive(FRDGBuilder& GraphBuilder, const FSceneView& InView,
-		const FPostProcessMaterialInputs& Inputs, ISceneViewExtension::EPostProcessingPass InLocation,
-		const TArray<const FCompositeCorePassProxy*> InPasses, const TSortedMap<int32, FScreenPassTextureSlice>& InputOverrides,
-		const FIntRect& SceneColorRect, int32 Level, FScreenPassTexture& Output);
+	/* Convenience function to create an override map of external textures. */
+	FExternalTextureMap CreateExternalTextureMap(FRDGBuilder& GraphBuilder, const FSceneView& InView, const FPostProcessMaterialInputs& Inputs) const;
+	
+	/* Convenience function to update the next pass inputs from the last pass' output. */
+	void UpdateNextPassInputs(
+		FRDGBuilder& GraphBuilder,
+		const UE::CompositeCore::FPassOutput& InOutput,
+		const int32 BindingIndex,
+		UE::CompositeCore::FPassInputArray& InOutInputs,
+		FExternalTextureMap& InOutExternalTextures
+	);
+
+	/* Recursive pass application. */
+	bool ApplyPasses_Recursive(
+		FRDGBuilder& GraphBuilder,
+		const FSceneView& InView,
+		const UE::CompositeCore::FPassInputArray& Inputs,
+		const UE::CompositeCore::FPassInputArray& OriginalInputs,
+		UE::CompositeCore::FPassContext& PassContext,
+		const TArray<const FCompositeCorePassProxy*> InPasses,
+		FExternalTextureMap& InOutExternalTextures,
+		int32 RecursionLevel,
+		UE::CompositeCore::FPassOutput& Output
+	);
 
 	/* Callback for processing passes. */
 	FScreenPassTexture PostProcessWork_RenderThread(FRDGBuilder& GraphBuilder, const FSceneView& InView,
@@ -85,7 +120,7 @@ private:
 	UE::CompositeCore::FBuiltInRenderPassOptions BuiltInRenderPassOptions;
 
 	/* Array of external texture inputs. */
-	TMap<UE::CompositeCore::ResourceId, TRefCountPtr<IPooledRenderTarget>> ExternalInputs_RenderThread;
+	TMap<UE::CompositeCore::ResourceId, UE::CompositeCore::FExternalRenderTarget> ExternalInputs_RenderThread;
 
 	/* Custom render pass render targets for each active view .*/
 	TMap<uint32, TRefCountPtr<IPooledRenderTarget>> CustomRenderTargetPerView_RenderThread;
diff --git a/Engine/Plugins/Experimental/Compositing/CompositeCore/Source/CompositeCore/Private/Passes/CompositeCorePassMergeProxy.cpp b/Engine/Plugins/Experimental/Compositing/CompositeCore/Source/CompositeCore/Private/Passes/CompositeCorePassMergeProxy.cpp
--- a/Engine/Plugins/Experimental/Compositing/CompositeCore/Source/CompositeCore/Private/Passes/CompositeCorePassMergeProxy.cpp
+++ b/Engine/Plugins/Experimental/Compositing/CompositeCore/Source/CompositeCore/Private/Passes/CompositeCorePassMergeProxy.cpp
@@ -54,13 +54,6 @@ namespace UE
 {
 	namespace CompositeCore
 	{
-		enum class ESourceEncoding : uint32
-		{
-			Linear = 0,
-			Gamma = 1,
-			sRGB = 2,
-		};
-
 		enum class EDistortionUV
 		{
 			None = 0,
@@ -69,40 +62,24 @@ namespace UE
 		};
 
 		FCompositeCoreMergeShader::FInputPassTextureParameters GetPassTextureParameters(
-			const FScreenPassTexture& InTexture,
-			const UE::CompositeCore::FPassInput PassInput,
+			const UE::CompositeCore::FPassInput& Input,
+			const FPassContext& PassContext,
 			const FSceneViewFamily* Family,
-			const ISceneViewExtension::EPostProcessingPass InLocation,
 			const bool bUseLensDistortion)
 		{
 			FCompositeCoreMergeShader::FInputPassTextureParameters Parameters;
 
-			Parameters.Texture = InTexture.Texture;
+			Parameters.Texture = Input.Texture.Texture;
 			Parameters.Sampler = TStaticSamplerState<SF_Bilinear>::GetRHI();
-			Parameters.bInvertedAlpha = PassInput.IsSet() ? PassInput.GetValue().bInvertedAlpha : true;
-			Parameters.SourceEncoding = static_cast<uint32>(ESourceEncoding::Linear);
-
-			// After-tonemap scene color may have encoding manually applied, as opposed to _SRGB textures
-			if (!PassInput.IsSet() && InLocation >= ISceneViewExtension::EPostProcessingPass::Tonemap)
-			{
-				if ((Family->EngineShowFlags.Tonemapper == 0) || (Family->EngineShowFlags.PostProcessing == 0))
-				{
-					Parameters.SourceEncoding = static_cast<uint32>(ESourceEncoding::Gamma);
-				}
-				else if (Family->SceneCaptureSource == SCS_FinalColorLDR)
-				{
-					Parameters.SourceEncoding = static_cast<uint32>(ESourceEncoding::sRGB);
-				}
-			}
+			Parameters.bInvertedAlpha = Input.Metadata.bInvertedAlpha;
+			Parameters.SourceEncoding = static_cast<uint32>(Input.Metadata.Encoding);
 
 			EDistortionUV DistortionUV = EDistortionUV::None;
 
 			if (bUseLensDistortion)
 			{
-				const bool bInputIsDistorted = PassInput.IsSet() ? PassInput.GetValue().bDistorted
-					: InLocation >= ISceneViewExtension::EPostProcessingPass::SSRInput;
-
-				const bool bOutputIsDistorted = InLocation <= ISceneViewExtension::EPostProcessingPass::SSRInput;
+				const bool bInputIsDistorted = Input.Metadata.bDistorted;
+				const bool bOutputIsDistorted = PassContext.bOutputDistortionEnabled;
 
 				// Distortion mismatch?
 				if (bInputIsDistorted != bOutputIsDistorted)
@@ -117,53 +94,50 @@ namespace UE
 			return Parameters;
 		}
 
-		FMergePassProxy::FMergePassProxy(UE::CompositeCore::FPassInputArray InPassInputs)
-			: FCompositeCorePassProxy(MoveTemp(InPassInputs))
-		{
-		}
-
-		FScreenPassTexture FMergePassProxy::Add(FRDGBuilder& GraphBuilder, const FSceneView& InView, const FPostProcessMaterialInputs& Inputs, const FIntRect& OutputViewRect, ISceneViewExtension::EPostProcessingPass InLocation) const
+		FPassOutput FMergePassProxy::Add(FRDGBuilder& GraphBuilder, const FSceneView& InView, const FPassInputArray& Inputs, const FPassContext& PassContext) const
 		{
 			RDG_EVENT_SCOPE_STAT(GraphBuilder, FCompositeCoreMerge, "CompositeCore.Merge");
 			RDG_GPU_STAT_SCOPE(GraphBuilder, FCompositeCoreMerge);
 
-			FScreenPassTexture Input0 = FScreenPassTexture::CopyFromSlice(GraphBuilder, Inputs.GetInput(static_cast<EPostProcessMaterialInput>(0)));
-			FScreenPassTexture Input1 = FScreenPassTexture::CopyFromSlice(GraphBuilder, Inputs.GetInput(static_cast<EPostProcessMaterialInput>(1)));
-
 			FGlobalShaderMap* GlobalShaderMap = GetGlobalShaderMap(InView.GetFeatureLevel());
 			const FSceneViewFamily* Family = InView.Family;
 
-			ESourceEncoding OutputEncoding = ESourceEncoding::Linear;
+			const FLensDistortionLUT& LensDistortionLUT = LensDistortion::GetLUTUnsafe(InView);
+			const bool bUseLensDistortion = LensDistortionLUT.IsEnabled() && (LensDistortion::GetPassLocationUnsafe(InView) == LensDistortion::EPassLocation::TSR);
+
+			FResourceMetadata OutMetadata = {};
 
-			// When running after tonemap, we assume that input0 (scene color) has an encoding.
-			if (InLocation >= ISceneViewExtension::EPostProcessingPass::Tonemap)
+			if (PassContext.bOutputSceneColor)
 			{
-				if ((Family->EngineShowFlags.Tonemapper == 0) || (Family->EngineShowFlags.PostProcessing == 0))
+				OutMetadata.bInvertedAlpha = true;
+				OutMetadata.bDistorted = PassContext.bOutputDistortionEnabled && bUseLensDistortion;
+				if (PassContext.Location >= ISceneViewExtension::EPostProcessingPass::Tonemap)
 				{
-					OutputEncoding = ESourceEncoding::Gamma;
-				}
-				else if (Family->SceneCaptureSource == SCS_FinalColorLDR)
-				{
-					OutputEncoding = ESourceEncoding::sRGB;
+					if ((Family->EngineShowFlags.Tonemapper == 0) || (Family->EngineShowFlags.PostProcessing == 0))
+					{
+						OutMetadata.Encoding = EEncoding::Gamma;
+					}
+					else if (Family->SceneCaptureSource == SCS_FinalColorLDR)
+					{
+						OutMetadata.Encoding = EEncoding::sRGB;
+					}
 				}
 			}
 
 			FScreenPassRenderTarget Output = Inputs.OverrideOutput;
 			if (!Output.IsValid())
 			{
-				Output = CreateOutputRenderTarget(GraphBuilder, InView, OutputViewRect, Input0.Texture->Desc, TEXT("CompositeCoreMergeOutput"));
+				const FRDGTextureDesc Desc = Inputs[0].Texture.Texture->Desc;
+				Output = CreateOutputRenderTarget(GraphBuilder, InView, PassContext.OutputViewRect, Desc, TEXT("CompositeCoreMergeOutput"));
 			}
 
-			const FLensDistortionLUT& LensDistortionLUT = LensDistortion::GetLUTUnsafe(InView);
-			const bool bUseLensDistortion = LensDistortionLUT.IsEnabled() && (LensDistortion::GetPassLocationUnsafe(InView) == LensDistortion::EPassLocation::TSR);
-
 			FCompositeCoreMergeShader::FParameters* PassParameters = GraphBuilder.AllocParameters<FCompositeCoreMergeShader::FParameters>();
 			PassParameters->View = InView.ViewUniformBuffer;
-			PassParameters->Input0 = GetScreenPassTextureViewportParameters(FScreenPassTextureViewport(Input0));
-			PassParameters->Input1 = GetScreenPassTextureViewportParameters(FScreenPassTextureViewport(Input1));
+			PassParameters->Input0 = GetScreenPassTextureViewportParameters(FScreenPassTextureViewport(Inputs[0].Texture));
+			PassParameters->Input1 = GetScreenPassTextureViewportParameters(FScreenPassTextureViewport(Inputs[1].Texture));
 			PassParameters->Output = GetScreenPassTextureViewportParameters(FScreenPassTextureViewport(Output));
-			PassParameters->Tex0 = GetPassTextureParameters(Input0, PassInputs[0], Family, InLocation, bUseLensDistortion);
-			PassParameters->Tex1 = GetPassTextureParameters(Input1, PassInputs[1], Family, InLocation, bUseLensDistortion);
+			PassParameters->Tex0 = GetPassTextureParameters(Inputs[0], PassContext, Family, bUseLensDistortion);
+			PassParameters->Tex1 = GetPassTextureParameters(Inputs[1], PassContext, Family, bUseLensDistortion);
 			if (bUseLensDistortion)
 			{
 				PassParameters->DistortingDisplacementTexture = LensDistortionLUT.DistortingDisplacementTexture;
@@ -178,8 +152,8 @@ namespace UE
 			PassParameters->DistortingDisplacementSampler = TStaticSamplerState<SF_Bilinear>::GetRHI();
 			PassParameters->DisplayGamma = FVector2f(Family->RenderTarget->GetDisplayGamma(), 1.0f / Family->RenderTarget->GetDisplayGamma());
 			PassParameters->MergeOp = static_cast<uint32>(MergeOp);
-			PassParameters->OutputEncoding = static_cast<uint32>(OutputEncoding);
-			PassParameters->bFlags = static_cast<uint32>(bUnpremulInput1 << 1) | static_cast<uint32>(bOutputSceneColor);
+			PassParameters->OutputEncoding = static_cast<uint32>(OutMetadata.Encoding);
+			PassParameters->bFlags = static_cast<uint32>(bUnpremulInput1 << 1) | static_cast<uint32>(PassContext.bOutputSceneColor);
 			PassParameters->RenderTargets[0] = Output.GetRenderTargetBinding();
 
 			FCompositeCoreMergeShader::FPermutationDomain PermutationVector;
@@ -196,7 +170,7 @@ namespace UE
 				Output.ViewRect
 			);
 
-			return MoveTemp(Output);
+			return FPassOutput{ MoveTemp(Output), MoveTemp(OutMetadata), PassOutputOverride };
 		}
 	}
 }
diff --git a/Engine/Plugins/Experimental/Compositing/CompositeCore/Source/CompositeCore/Private/Passes/CompositeCorePassProxy.cpp b/Engine/Plugins/Experimental/Compositing/CompositeCore/Source/CompositeCore/Private/Passes/CompositeCorePassProxy.cpp
--- a/Engine/Plugins/Experimental/Compositing/CompositeCore/Source/CompositeCore/Private/Passes/CompositeCorePassProxy.cpp
+++ b/Engine/Plugins/Experimental/Compositing/CompositeCore/Source/CompositeCore/Private/Passes/CompositeCorePassProxy.cpp
@@ -7,12 +7,18 @@
 #include "FXRenderingUtils.h"
 #include "RenderGraphBuilder.h"
 #include "RenderGraphUtils.h"
-#include "ScreenPass.h"
+#include "PostProcess/PostProcessMaterialInputs.h"
 
 namespace UE
 {
 	namespace CompositeCore
 	{
+		FPassOutput::FPassOutput(FScreenPassTexture InTexture, FResourceMetadata InMetadata, TOptional<ResourceId> InOverride)
+			: Resource(MoveTemp(InTexture), MoveTemp(InMetadata))
+			, Override(InOverride)
+		{
+		}
+
 		FBuiltInRenderPassOptions::FBuiltInRenderPassOptions()
 			: ViewUserFlagsOverride{}
 			, bEnableUnlitViewmode{true}
@@ -28,18 +34,18 @@ namespace UE
 					TArray<const FCompositeCorePassProxy*>& SSRInputWork = RenderWork.FramePasses.Add(ISceneViewExtension::EPostProcessingPass::SSRInput);
 					TArray<const FCompositeCorePassProxy*>& AfterTonemapWork = RenderWork.FramePasses.Add(ISceneViewExtension::EPostProcessingPass::Tonemap);
 
-					FPassInputArray PassInputs;
-					// Input 0: Default
-					PassInputs.Add({});
+					// Default primary input
+					FPassInputDecl Input0;
+					Input0.Set<FPassInternalResourceDesc>({});
 					
-					// Input 1: Built-in custom render pass
-					FPassResourceDesc Desc;
-					Desc.Id = BUILT_IN_CRP_ID;
-					Desc.bInvertedAlpha = true;
-					PassInputs.Add(Desc);
+					// Built-in custom render pass as secondary input
+					FPassInputDecl Input1;
+					Input1.Set<FPassExternalResourceDesc>({ BUILT_IN_CRP_ID });
 
-					SSRInputWork.Add(RenderWork.FrameAllocator->Create<FMergePassProxy>(PassInputs));
-					AfterTonemapWork.Add(RenderWork.FrameAllocator->Create<FMergePassProxy>(PassInputs));
+					const FPassInputDeclArray PassDeclaredInputs = { Input0, Input1 };
+
+					SSRInputWork.Add(RenderWork.FrameAllocator->Create<FMergePassProxy>(PassDeclaredInputs));
+					AfterTonemapWork.Add(RenderWork.FrameAllocator->Create<FMergePassProxy>(PassDeclaredInputs));
 
 					return RenderWork;
 				}();
@@ -54,6 +60,21 @@ namespace UE
 	}
 }
 
+
+FCompositeCorePassProxy::FCompositeCorePassProxy(UE::CompositeCore::FPassInputDeclArray InPassDeclaredInputs, TOptional<UE::CompositeCore::ResourceId> InPassOutputOverride)
+	: PassDeclaredInputs(MoveTemp(InPassDeclaredInputs))
+	, PassOutputOverride(MoveTemp(InPassOutputOverride))
+{ }
+
+
+void FCompositeCorePassProxy::AddSubPasses(int32 InputIndex, TArray<const FCompositeCorePassProxy*> InSubPasses)
+{
+	if (!InSubPasses.IsEmpty())
+	{
+		SubPasses.Emplace(InputIndex, MoveTemp(InSubPasses));
+	}
+}
+
 FScreenPassRenderTarget FCompositeCorePassProxy::CreateOutputRenderTarget(FRDGBuilder& GraphBuilder,
 	const FSceneView& InView, const FIntRect& OutputViewRect, FRDGTextureDesc OutputDesc, const TCHAR* InName)
 {
@@ -65,3 +86,75 @@ FScreenPassRenderTarget FCompositeCorePassProxy::CreateOutputRenderTarget(FRDGBu
 
 	return FScreenPassRenderTarget(GraphBuilder.CreateTexture(OutputDesc, InName), OutputViewRect, InView.GetOverwriteLoadAction());
 }
+
+UE::CompositeCore::FPassInputArray::FPassInputArray(
+	FRDGBuilder& GraphBuilder,
+	const FSceneView& InView,
+	const FPostProcessMaterialInputs& InPostInputs,
+	const ISceneViewExtension::EPostProcessingPass& InLocation,
+	const bool bLensDistortionEnabled)
+{
+	using namespace UE::CompositeCore;
+
+	Inputs.Reserve(kPostProcessMaterialInputCountMax);
+
+	{
+		FPassInput& Input = Inputs.AddDefaulted_GetRef();
+		Input.Texture = FScreenPassTexture::CopyFromSlice(GraphBuilder, InPostInputs.GetInput(EPostProcessMaterialInput::SceneColor));
+		Input.Metadata.bInvertedAlpha = true;
+		Input.Metadata.bDistorted = bLensDistortionEnabled && InLocation >= ISceneViewExtension::EPostProcessingPass::SSRInput;
+
+		// After-tonemap scene color may have encoding manually applied, as opposed to _SRGB textures
+		if (InLocation >= ISceneViewExtension::EPostProcessingPass::Tonemap)
+		{
+			if ((InView.Family->EngineShowFlags.Tonemapper == 0) || (InView.Family->EngineShowFlags.PostProcessing == 0))
+			{
+				Input.Metadata.Encoding = UE::CompositeCore::EEncoding::Gamma;
+			}
+			else if (InView.Family->SceneCaptureSource == SCS_FinalColorLDR)
+			{
+				Input.Metadata.Encoding = UE::CompositeCore::EEncoding::sRGB;
+			}
+		}
+	}
+
+	{
+		FPassInput& Input = Inputs.AddDefaulted_GetRef();
+		Input.Texture = FScreenPassTexture::CopyFromSlice(GraphBuilder, InPostInputs.GetInput(EPostProcessMaterialInput::SeparateTranslucency));
+		Input.Metadata.bInvertedAlpha = true;
+	}
+
+	{
+		FPassInput& Input = Inputs.AddDefaulted_GetRef();
+		Input.Texture = FScreenPassTexture::CopyFromSlice(GraphBuilder, InPostInputs.GetInput(EPostProcessMaterialInput::CombinedBloom));
+		Input.Metadata.bInvertedAlpha = true;
+	}
+
+	{
+		FPassInput& Input = Inputs.AddDefaulted_GetRef();
+		Input.Texture = FScreenPassTexture::CopyFromSlice(GraphBuilder, InPostInputs.GetInput(EPostProcessMaterialInput::PostTonemapHDRColor));
+		Input.Metadata.bInvertedAlpha = true;
+	}
+
+	{
+		FPassInput& Input = Inputs.AddDefaulted_GetRef();
+		Input.Texture = FScreenPassTexture::CopyFromSlice(GraphBuilder, InPostInputs.GetInput(EPostProcessMaterialInput::Velocity));
+	}
+
+	OverrideOutput = InPostInputs.OverrideOutput;
+}
+
+FPostProcessMaterialInputs UE::CompositeCore::FPassInputArray::ToPostProcessInputs(FRDGBuilder& GraphBuilder, FSceneTextureShaderParameters SceneTextures) const
+{
+	using namespace UE::CompositeCore;
+
+	FPostProcessMaterialInputs Result;
+	for (int32 Index = 0; Index < Inputs.Num(); ++Index)
+	{
+		const FPassInput& ResolvedInput = Inputs[Index];
+		Result.SetInput(static_cast<EPostProcessMaterialInput>(Index), FScreenPassTextureSlice::CreateFromScreenPassTexture(GraphBuilder, ResolvedInput.Texture));
+	}
+	Result.SceneTextures = MoveTemp(SceneTextures);
+
+	return Result;
+}
diff --git a/Engine/Plugins/Experimental/Compositing/CompositeCore/Source/CompositeCore/Public/Passes/CompositeCorePassMergeProxy.h b/Engine/Plugins/Experimental/Compositing/CompositeCore/Source/CompositeCore/Public/Passes/CompositeCorePassMergeProxy.h
--- a/Engine/Plugins/Experimental/Compositing/CompositeCore/Source/CompositeCore/Public/Passes/CompositeCorePassMergeProxy.h
+++ b/Engine/Plugins/Experimental/Compositing/CompositeCore/Source/CompositeCore/Public/Passes/CompositeCorePassMergeProxy.h
@@ -10,6 +10,7 @@
 UENUM()
 enum class ECompositeCoreMergeOp : uint8
 {
+	None,
 	Over,
 	Add,
 	Multiply,
@@ -23,19 +24,15 @@ namespace UE
 		class COMPOSITECORE_API FMergePassProxy : public FCompositeCorePassProxy
 		{
 		public:
-			/** Constructor. */
-			FMergePassProxy(UE::CompositeCore::FPassInputArray InPassInputs);
+			using FCompositeCorePassProxy::FCompositeCorePassProxy;
 
-			FScreenPassTexture Add(FRDGBuilder& GraphBuilder, const FSceneView& InView, const FPostProcessMaterialInputs& Inputs, const FIntRect& OutputViewRect, ISceneViewExtension::EPostProcessingPass InLocation) const override;
+			FPassOutput Add(FRDGBuilder& GraphBuilder, const FSceneView& InView, const FPassInputArray& Inputs, const FPassContext& PassContext) const override;
 
 			/** Merge operation. */
 			ECompositeCoreMergeOp MergeOp = ECompositeCoreMergeOp::Over;
 
-			/* Unpremultiplication required for recovering opaque colors for the default over blend. */
+			/** Unpremultiplication required for recovering opaque colors for the default over blend. */
 			bool bUnpremulInput1 = true;
-
-			/* Output is assumed to be scene color with inverted alpha. */
-			bool bOutputSceneColor = true;
 		};
 	}
 }
diff --git a/Engine/Plugins/Experimental/Compositing/CompositeCore/Source/CompositeCore/Public/Passes/CompositeCorePassProxy.h b/Engine/Plugins/Experimental/Compositing/CompositeCore/Source/CompositeCore/Public/Passes/CompositeCorePassProxy.h
--- a/Engine/Plugins/Experimental/Compositing/CompositeCore/Source/CompositeCore/Public/Passes/CompositeCorePassProxy.h
+++ b/Engine/Plugins/Experimental/Compositing/CompositeCore/Source/CompositeCore/Public/Passes/CompositeCorePassProxy.h
@@ -3,10 +3,14 @@
 #pragma once
 
 #include "CoreTypes.h"
+
+#include "Containers/Array.h"
 #include "Containers/ContainersFwd.h"
-#include "SceneViewExtension.h"
 #include "Misc/Optional.h"
-#include "Containers/Array.h"
+#include "Misc/TVariant.h"
+#include "SceneViewExtension.h"
+#include "SceneTexturesConfig.h"
+#include "ScreenPass.h"
 
 class FRDGBuilder;
 struct FRDGTextureDesc;
@@ -15,32 +19,179 @@ struct FScreenPassTexture;
 struct FScreenPassRenderTarget;
 struct FPostProcessMaterialInputs;
 class FCompositeCorePassProxy;
+struct IPooledRenderTarget;
+class UTexture;
+
+enum class EPostProcessMaterialInput : uint32;
 
 namespace UE
 {
 	namespace CompositeCore
 	{
-		/* Type used to identify passes, textures or built-in renderer sources/targets. */
+		/** Texture encoding type, used for scene color. (HDR is not currently supported.) */
+		enum class EEncoding : uint32
+		{
+			Linear = 0,
+			Gamma = 1,
+			sRGB = 2,
+		};
+
+		/** Type used to identify passes, textures or built-in renderer sources/targets. */
 		using ResourceId = uint32;
-		/* Built-in custom render pass identifier. */
+		/** Built-in custom render pass identifier. */
 		static constexpr ResourceId BUILT_IN_CRP_ID = 1;
-		/* Built-in empty/black identifier. */
+		/** Built-in empty/black identifier. */
 		static constexpr ResourceId BUILT_IN_EMPTY_ID = 2;
-		/* First identifier of the external texture inputs, see ExternalInputs on FPostRenderWork. */
+		/** First identifier of the external texture inputs, see ExternalInputs on FPostRenderWork. */
 		static constexpr ResourceId EXTERNAL_RANGE_START_ID = 100;
 
-		/** Pass texture description.	*/
-		struct FPassResourceDesc
+		/** Texture resource metadata. */
+		struct FResourceMetadata
+		{
+			/** Is the alpha inverted (like scene color)? */
+			bool bInvertedAlpha = false;
+
+			/** Is the texture content distorted? */
+			bool bDistorted = false;
+			
+			/** Source color encoding. */
+			EEncoding Encoding = EEncoding::Linear;
+		};
+
+		/** Pass texture description for internal resources (default scene textures). */
+		struct FPassInternalResourceDesc
+		{
+			/** Index, which maps to the default 0-4 post-processing inputs or beyond. */
+			int32 Index = 0;
+			
+			/** Flag to bypass the previous pass textures & access the original scene textures. */
+			bool bOriginalCopyBeforePasses = false;
+		};
+
+		/** Pass texture description for external render targets. */
+		struct FPassExternalResourceDesc
+		{
+			/** External resource identifier. */
+			ResourceId Id = 0;
+		};
+
+		/** Pass input declaration, referring to either internally or externally managed textures. */
+		using FPassInputDecl = TVariant<FPassInternalResourceDesc, FPassExternalResourceDesc>;
+		
+		/** Array of pass input declarations. */
+		using FPassInputDeclArray = TArray<FPassInputDecl>;
+
+		/** External texture resource and its accompanying metadata. */
+		struct FExternalTexture
+		{
+			/** Texture weak object pointer, for use on the game thread. */
+			TWeakObjectPtr<UTexture> Texture = {};
+			
+			/** Texture metadata. */
+			FResourceMetadata Metadata = {};
+		};
+
+		/** Resolved texture resource with an active (screen) texture and its accompanying metadata. */
+		struct FPassTexture
 		{
-			ResourceId Id			= 0;
-			bool bInvertedAlpha		= false;
-			bool bDistorted			= false;
+			/** Pass screen texture. */
+			FScreenPassTexture Texture = {};
+			
+			/** Texture metadata. */
+			FResourceMetadata Metadata = {};
 		};
 
-		/* Optional pass input override, from the otherwise default 0-4 post-processing inputs. */
-		using FPassInput = TOptional<FPassResourceDesc>;
-		/* Array of pass inputs. */
-		using FPassInputArray = TArray<FPassInput>;
+		/** Pass texture input definition. */
+		using FPassInput = FPassTexture;
+
+		/** Resolved pass inputs. */
+		struct COMPOSITECORE_API FPassInputArray
+		{
+			/** Default constructor */
+			FPassInputArray() = default;
+			FPassInputArray(FPassInputArray&&) = default;
+			FPassInputArray(const FPassInputArray&) = default;
+			FPassInputArray& operator=(FPassInputArray&&) = default;
+			FPassInputArray& operator=(const FPassInputArray&) = default;
+
+			/** Constructor */
+			FPassInputArray(
+				FRDGBuilder& GraphBuilder,
+				const FSceneView& InView,
+				const FPostProcessMaterialInputs& InPostInputs,
+				const ISceneViewExtension::EPostProcessingPass& InLocation,
+				const bool bInIsDistorted
+			);
+
+			/** Array access operator overload */
+			const FPassInput& operator[] (int32 Index) const
+			{
+				return Inputs[Index];
+			}
+
+			/** Array access operator overload */
+			FPassInput& operator[] (int32 Index)
+			{
+				return Inputs[Index];
+			}
+			
+			/** Tests if index is valid, i.e. greater than or equal to zero, and less than the number of elements in the array. */
+			bool IsValidIndex(int32 Index) const
+			{
+				return Inputs.IsValidIndex(Index);
+			}
+
+			/** Number of pass inputs */
+			int32 Num() const { return Inputs.Num(); }
+
+			/** Conversion function to engine post-process (material) inputs. */
+			FPostProcessMaterialInputs ToPostProcessInputs(FRDGBuilder& GraphBuilder, FSceneTextureShaderParameters SceneTextures) const;
+
+			/** Pass override output. */
+			FScreenPassRenderTarget OverrideOutput;
+
+		private:
+			/** Pass texture input array. */
+			TArray<UE::CompositeCore::FPassInput> Inputs;
+		};
+
+		/** Pass output definition. */
+		struct COMPOSITECORE_API FPassOutput
+		{
+			/** Default operators */
+			FPassOutput() = default;
+			FPassOutput(FPassOutput&&) = default;
+			FPassOutput(const FPassOutput&) = default;
+			FPassOutput& operator=(FPassOutput&&) = default;
+			FPassOutput& operator=(const FPassOutput&) = default;
+
+			/** Constructor */
+			explicit FPassOutput(FScreenPassTexture InTexture, FResourceMetadata InMetadata, TOptional<ResourceId> InOverride = {});
+
+			/** Output pass resource. */
+			FPassTexture Resource = {};
+			/** Optional pass output target override, from the otherwise default scene color. */
+			TOptional<ResourceId> Override = {};
+		};
+
+		/** Pass parameter information. */
+		struct FPassContext
+		{
+			/** The uniform buffer containing all scene textures. */
+			FSceneTextureShaderParameters SceneTextures;
+
+			/** Active post-processing output view rectangle. */
+			FIntRect OutputViewRect;
+			
+			/** Post-processiong location. */
+			ISceneViewExtension::EPostProcessingPass Location;
+			
+			/** Is the current pass expected to output scene color? */
+			bool bOutputSceneColor = false;
+			
+			/** Is the current pass expected to output (lens) distorted data? */
+			bool bOutputDistortionEnabled = false;
+		};
 
 		/** Options to control the built-in custom render pass. */
 		struct COMPOSITECORE_API FBuiltInRenderPassOptions final
@@ -48,13 +199,13 @@ namespace UE
 			/** Constructor. */
 			FBuiltInRenderPassOptions();
 
-			/* Custom user flags value used to alter materials in the composite render pass. */
+			/** Custom user flags value used to alter materials in the composite render pass. */
 			TOptional<int32> ViewUserFlagsOverride;
 
 			/** Enables the development shader debug feature that routes the Base Color output to Emissive for the separate render. Non-shipping PC build only. */
 			bool bEnableUnlitViewmode;
 
-			/* Enables the dilation pass with size > 0 (currently only 1 is supported). */
+			/** Enables the dilation pass with size > 0 (currently only 1 is supported). */
 			int32 OutputDilationSize;
 		};
 
@@ -67,8 +218,8 @@ namespace UE
 			/** Constructor. */
 			FPostRenderWork();
 
-			/* Array of user-defined external input texture overrides, where indices map to ResourceId. */
-			TArray<const TWeakObjectPtr<UTexture>> ExternalInputs;
+			/** Array of user-defined external input texture overrides, where indices map to ResourceId. */
+			TArray<FExternalTexture> ExternalInputs;
 
 			/** Post-processing passes at the specified locations. */
 			TSortedMap<ISceneViewExtension::EPostProcessingPass, TArray<const FCompositeCorePassProxy*>> FramePasses;
@@ -83,55 +234,57 @@ namespace UE
 class COMPOSITECORE_API FCompositeCorePassProxy
 {
 public:
-
 	/** Default constructor. */
 	FCompositeCorePassProxy() = default;
 
 	/** Constructor. */
-	FCompositeCorePassProxy(UE::CompositeCore::FPassInputArray InPassInputs)
-		: PassInputs(MoveTemp(InPassInputs))
-	{ }
+	FCompositeCorePassProxy(UE::CompositeCore::FPassInputDeclArray InPassDeclaredInputs, TOptional<UE::CompositeCore::ResourceId> InPassOutputOverride = {});
 
-	/* Virtual destructor */
+	/** Virtual destructor */
 	virtual ~FCompositeCorePassProxy() { };
 
-	/* Render-thread add pass method to override. */
-	virtual FScreenPassTexture Add(FRDGBuilder& GraphBuilder, const FSceneView& InView, const FPostProcessMaterialInputs& Inputs, const FIntRect& OutputViewRect, ISceneViewExtension::EPostProcessingPass InLocation) const = 0;
+	/** Render-thread add pass method to override. */
+	virtual UE::CompositeCore::FPassOutput Add(
+		FRDGBuilder& GraphBuilder,
+		const FSceneView& InView,
+		const UE::CompositeCore::FPassInputArray& Inputs,
+		const UE::CompositeCore::FPassContext& PassContext
+	) const = 0;
 
-	/* Number of inputs used by the pass. */
-	int32 GetNumInputs() const
+	/** Number of inputs used by the pass. */
+	int32 GetNumDeclaredInputs() const
 	{
-		return PassInputs.Num();
+		return PassDeclaredInputs.Num();
 	}
 
-	/* Number of inputs used by the pass. */
-	const UE::CompositeCore::FPassInput& GetPassInput(int32 InputIndex) const
+	/** Get pass input at specified index. */
+	const UE::CompositeCore::FPassInputDecl& GetDeclaredInput(int32 InputIndex) const
 	{
-		return PassInputs[InputIndex];
+		return PassDeclaredInputs[InputIndex];
 	}
 
-	/* Return sub-passes per input index. */
+	/** Return sub-passes per input index. */
 	const TArray<const FCompositeCorePassProxy*>* GetSubPasses(int32 InputIndex) const
 	{
 		return SubPasses.Find(InputIndex);
 	}
 
-	/* Add sub-passes per input index. */
-	void AddSubPasses(int32 InputIndex, TArray<const FCompositeCorePassProxy*> InSubPasses)
-	{
-		SubPasses.Emplace(InputIndex, MoveTemp(InSubPasses));
-	}
+	/** Add sub-passes per input index. */
+	void AddSubPasses(int32 InputIndex, TArray<const FCompositeCorePassProxy*> InSubPasses);
 
-	/* Convenience function to create an output render target with the specified resolution. */
+	/** Convenience function to create an output render target with the specified resolution. */
 	static FScreenPassRenderTarget CreateOutputRenderTarget(FRDGBuilder& GraphBuilder,
 		const FSceneView& InView, const FIntRect& OutputViewRect, FRDGTextureDesc OutputDesc, const TCHAR* InName);
 
 protected:
 
-	/* List of pass input types. */
-	UE::CompositeCore::FPassInputArray PassInputs;
+	/** List of pass input types. */
+	UE::CompositeCore::FPassInputDeclArray PassDeclaredInputs;
+
+	/** Optional pass output override. */
+	TOptional<UE::CompositeCore::ResourceId> PassOutputOverride;
 
-	/* Map of sub-passes per input index. */
+	/** Map of sub-passes per input index. */
 	TSortedMap<int32, TArray<const FCompositeCorePassProxy*>> SubPasses;
 };
 
