CompositeCore: Fix PCH compilation error on android.

diff --git a/Engine/Plugins/Experimental/Compositing/CompositeCore/Source/CompositeCore/Private/Passes/CompositeCorePassProxy.cpp b/Engine/Plugins/Experimental/Compositing/CompositeCore/Source/CompositeCore/Private/Passes/CompositeCorePassProxy.cpp
--- a/Engine/Plugins/Experimental/Compositing/CompositeCore/Source/CompositeCore/Private/Passes/CompositeCorePassProxy.cpp
+++ b/Engine/Plugins/Experimental/Compositing/CompositeCore/Source/CompositeCore/Private/Passes/CompositeCorePassProxy.cpp
@@ -14,7 +14,7 @@ namespace UE
 	namespace CompositeCore
 	{
 		FPassOutput::FPassOutput(FScreenPassTexture InTexture, FResourceMetadata InMetadata, TOptional<ResourceId> InOverride)
-			: Resource(MoveTemp(InTexture), MoveTemp(InMetadata))
+			: Resource({ MoveTemp(InTexture), MoveTemp(InMetadata) })
 			, Override(InOverride)
 		{
 		}
