CompositeCore: Additional android compilation fix.

Also include comment formatting fix.

diff --git a/Engine/Plugins/Experimental/Compositing/CompositeCore/Source/CompositeCore/Private/CompositeCoreSceneViewExtension.cpp b/Engine/Plugins/Experimental/Compositing/CompositeCore/Source/CompositeCore/Private/CompositeCoreSceneViewExtension.cpp
--- a/Engine/Plugins/Experimental/Compositing/CompositeCore/Source/CompositeCore/Private/CompositeCoreSceneViewExtension.cpp
+++ b/Engine/Plugins/Experimental/Compositing/CompositeCore/Source/CompositeCore/Private/CompositeCoreSceneViewExtension.cpp
@@ -356,7 +356,10 @@ TSortedMap<int32, UE::CompositeCore::FPassInput> FCompositeCoreSceneViewExtensio
 	TSortedMap<int32, FPassInput> InputOverrides;
 	InputOverrides.Reserve(2 + ExternalInputs_RenderThread.Num());
 
-	InputOverrides.Emplace(BUILT_IN_EMPTY_ID, FScreenPassTexture{ GSystemTextures.GetBlackDummy(GraphBuilder) });
+	{
+		FPassInput& Resource = InputOverrides.Add(BUILT_IN_EMPTY_ID);
+		Resource.Texture = FScreenPassTexture{ GSystemTextures.GetBlackDummy(GraphBuilder) };
+	}
 
 	const TRefCountPtr<IPooledRenderTarget>* CompositeRenderPassPtr = CustomRenderTargetPerView_RenderThread.Find(InView.GetViewKey());
 	if (CompositeRenderPassPtr != nullptr)
diff --git a/Engine/Plugins/Experimental/Compositing/CompositeCore/Source/CompositeCore/Private/CompositeCoreSceneViewExtension.h b/Engine/Plugins/Experimental/Compositing/CompositeCore/Source/CompositeCore/Private/CompositeCoreSceneViewExtension.h
--- a/Engine/Plugins/Experimental/Compositing/CompositeCore/Source/CompositeCore/Private/CompositeCoreSceneViewExtension.h
+++ b/Engine/Plugins/Experimental/Compositing/CompositeCore/Source/CompositeCore/Private/CompositeCoreSceneViewExtension.h
@@ -39,22 +39,22 @@ public:
 	FCompositeCoreSceneViewExtension(const FAutoRegister& AutoReg, UWorld* InWorld);
 	~FCompositeCoreSceneViewExtension();
 
-	/* Register primitives for compositing. */
+	/** Register primitives for compositing. */
 	void RegisterPrimitives_GameThread(const TArray<UPrimitiveComponent*>& InPrimitiveComponents);
 
-	/* Unregister primitives for compositing. */
+	/** Unregister primitives for compositing. */
 	void UnregisterPrimitives_GameThread(const TArray<UPrimitiveComponent*>& InPrimitiveComponents);
 
-	/* Set post-processing render work to be enqueued on the render thread. */
+	/** Set post-processing render work to be enqueued on the render thread. */
 	void SetPostRenderWork_GameThread(UE::CompositeCore::FPostRenderWork&& InWork);
 
-	/* Reset post-processing render work. */
+	/** Reset post-processing render work. */
 	void ResetPostRenderWork_GameThread();
 
-	/* Set an optional view user flags override. */
+	/** Set an optional view user flags override. */
 	void SetBuiltInRenderPassOptions_GameThread(const UE::CompositeCore::FBuiltInRenderPassOptions& InOptions);
 
-	/* Called by the custom render pass to store its view render target for this frame. */
+	/** Called by the custom render pass to store its view render target for this frame. */
 	void CollectCustomRenderTarget(uint32 InViewId, const TRefCountPtr<IPooledRenderTarget>& InRenderTarget);
 
 	//~ Begin ISceneViewExtension Interface
@@ -78,16 +78,16 @@ protected:
 private:
 	using FExternalTextureMap = TSortedMap<int32, UE::CompositeCore::FPassInput>;
 
-	/* Getter for the frame render work, either from the actor or default. */
+	/** Getter for the frame render work, either from the actor or default. */
 	const UE::CompositeCore::FPostRenderWork& GetRenderWork() const;
 
-	/* Active check to avoid work in post-processing. */
+	/** Active check to avoid work in post-processing. */
 	bool IsActiveForView(const FSceneView& InView) const;
 
-	/* Convenience function to create an override map of external textures. */
+	/** Convenience function to create an override map of external textures. */
 	FExternalTextureMap CreateExternalTextureMap(FRDGBuilder& GraphBuilder, const FSceneView& InView, const FPostProcessMaterialInputs& Inputs) const;
 	
-	/* Convenience function to update the next pass inputs from the last pass' output. */
+	/** Convenience function to update the next pass inputs from the last pass' output. */
 	void UpdateNextPassInputs(
 		FRDGBuilder& GraphBuilder,
 		const UE::CompositeCore::FPassOutput& InOutput,
@@ -96,7 +96,7 @@ private:
 		FExternalTextureMap& InOutExternalTextures
 	);
 
-	/* Recursive pass application. */
+	/** Recursive pass application. */
 	bool ApplyPasses_Recursive(
 		FRDGBuilder& GraphBuilder,
 		const FSceneView& InView,
@@ -109,23 +109,23 @@ private:
 		UE::CompositeCore::FPassOutput& Output
 	);
 
-	/* Callback for processing passes. */
+	/** Callback for processing passes. */
 	FScreenPassTexture PostProcessWork_RenderThread(FRDGBuilder& GraphBuilder, const FSceneView& InView,
 		const FPostProcessMaterialInputs& Inputs, ISceneViewExtension::EPostProcessingPass InLocation);
 
-	/* Collection of primitives to render as a custom render pass and composite after post-processing. .*/
+	/** Collection of primitives to render as a custom render pass and composite after post-processing. .*/
 	TSet<TWeakObjectPtr<UPrimitiveComponent>> CompositePrimitives;
 
-	/* Built-in composite render pass options. .*/
+	/** Built-in composite render pass options. .*/
 	UE::CompositeCore::FBuiltInRenderPassOptions BuiltInRenderPassOptions;
 
-	/* Array of external texture inputs. */
+	/** Array of external texture inputs. */
 	TMap<UE::CompositeCore::ResourceId, UE::CompositeCore::FExternalRenderTarget> ExternalInputs_RenderThread;
 
-	/* Custom render pass render targets for each active view .*/
+	/** Custom render pass render targets for each active view .*/
 	TMap<uint32, TRefCountPtr<IPooledRenderTarget>> CustomRenderTargetPerView_RenderThread;
 
-	/* Custom render work. */
+	/** Custom render work. */
 	TOptional<UE::CompositeCore::FPostRenderWork> PostRenderWork_RenderThread;
 };
 
