CompositeCore: Reintroduce "Composite Follows Scene Exposure" plugin setting which was present in 5.5, since this use case was recently reported.

diff --git a/Engine/Plugins/Experimental/Compositing/CompositeCore/Shaders/Private/CompositeCoreMerge.usf b/Engine/Plugins/Experimental/Compositing/CompositeCore/Shaders/Private/CompositeCoreMerge.usf
--- a/Engine/Plugins/Experimental/Compositing/CompositeCore/Shaders/Private/CompositeCoreMerge.usf
+++ b/Engine/Plugins/Experimental/Compositing/CompositeCore/Shaders/Private/CompositeCoreMerge.usf
@@ -30,7 +30,8 @@ SCREEN_PASS_TEXTURE_VIEWPORT(Output)
 	INPUT_PASS_TEXTURE_MEMBER(StructName, SamplerState, Sampler) \
     INPUT_PASS_TEXTURE_MEMBER(StructName, uint, bInvertedAlpha) \
     INPUT_PASS_TEXTURE_MEMBER(StructName, uint, SourceEncoding) \
-    INPUT_PASS_TEXTURE_MEMBER(StructName, uint, DistortionUV)
+    INPUT_PASS_TEXTURE_MEMBER(StructName, uint, DistortionUV) \
+    INPUT_PASS_TEXTURE_MEMBER(StructName, float, Exposure)
     
 
 INPUT_PASS_TEXTURE(Tex0);
@@ -157,6 +158,10 @@ void MainPS(
         const bool bUnpremulInput1 = ((bFlags >> 1) & 1);
         Color1 = select(bUnpremulInput1, Unpremultiply(Color1), Color1);
     }
+
+    // Adjust exposure if necessary
+    Color0.rgb *= Tex0_Exposure;
+    Color1.rgb *= Tex1_Exposure;
     
     if(MergeOp == ECompositeCoreMergeOp_Over)
     {
diff --git a/Engine/Plugins/Experimental/Compositing/CompositeCore/Source/CompositeCore/Private/CompositeCoreSettings.cpp b/Engine/Plugins/Experimental/Compositing/CompositeCore/Source/CompositeCore/Private/CompositeCoreSettings.cpp
--- a/Engine/Plugins/Experimental/Compositing/CompositeCore/Source/CompositeCore/Private/CompositeCoreSettings.cpp
+++ b/Engine/Plugins/Experimental/Compositing/CompositeCore/Source/CompositeCore/Private/CompositeCoreSettings.cpp
@@ -24,7 +24,8 @@ static FAutoConsoleVariableRef CVarRefreshPrimitivesOnTick(
 	ECVF_ReadOnly);
 
 UCompositeCorePluginSettings::UCompositeCorePluginSettings()
-	: DisabledPrimitiveClasses(
+	: bCompositeFollowsSceneExposure(false)
+	, DisabledPrimitiveClasses(
 		{
 			FSoftClassPath("/Script/Engine.CapsuleComponent"),
 			FSoftClassPath("/Script/Engine.BillboardComponent"),
@@ -60,6 +61,8 @@ void UCompositeCorePluginSettings::PostEditChangeProperty(FPropertyChangedEvent&
 	{
 		CacheDisabledClasses();
 	}
+
+	ExportValuesToConsoleVariables(PropertyChangedEvent.Property);
 }
 #endif
 
diff --git a/Engine/Plugins/Experimental/Compositing/CompositeCore/Source/CompositeCore/Private/Passes/CompositeCorePassMergeProxy.cpp b/Engine/Plugins/Experimental/Compositing/CompositeCore/Source/CompositeCore/Private/Passes/CompositeCorePassMergeProxy.cpp
--- a/Engine/Plugins/Experimental/Compositing/CompositeCore/Source/CompositeCore/Private/Passes/CompositeCorePassMergeProxy.cpp
+++ b/Engine/Plugins/Experimental/Compositing/CompositeCore/Source/CompositeCore/Private/Passes/CompositeCorePassMergeProxy.cpp
@@ -9,11 +9,18 @@
 #include "PostProcess/PostProcessMaterialInputs.h"
 #include "RenderGraphBuilder.h"
 #include "RenderGraphUtils.h"
+#include "SceneManagement.h"
 #include "SceneView.h"
 #include "SystemTextures.h"
 
 DECLARE_GPU_STAT_NAMED(FCompositeCoreMerge, TEXT("CompositeCore.Merge"));
 
+static TAutoConsoleVariable<bool> CVarCompositeCoreFollowSceneExposure(
+	TEXT("CompositeCore.FollowSceneExposure"),
+	false,
+	TEXT("When enabled, the scene main render pre-exposure is applied onto separate composited render(s)."),
+	ECVF_RenderThreadSafe);
+
 class FCompositeCoreMergeShader : public FGlobalShader
 {
 	DECLARE_GLOBAL_SHADER(FCompositeCoreMergeShader);
@@ -28,6 +35,7 @@ class FCompositeCoreMergeShader : public FGlobalShader
 		SHADER_PARAMETER(uint32, bInvertedAlpha)
 		SHADER_PARAMETER(uint32, SourceEncoding)
 		SHADER_PARAMETER(uint32, DistortionUV)
+		SHADER_PARAMETER(float, Exposure)
 	END_SHADER_PARAMETER_STRUCT()
 
 	BEGIN_SHADER_PARAMETER_STRUCT(FParameters, )
@@ -65,7 +73,8 @@ namespace UE
 			const UE::CompositeCore::FPassInput& Input,
 			const FPassContext& PassContext,
 			const FSceneViewFamily* Family,
-			const bool bUseLensDistortion)
+			const bool bUseLensDistortion,
+			const float InPreExposure)
 		{
 			FCompositeCoreMergeShader::FInputPassTextureParameters Parameters;
 
@@ -73,6 +82,16 @@ namespace UE
 			Parameters.Sampler = TStaticSamplerState<SF_Bilinear>::GetRHI();
 			Parameters.bInvertedAlpha = Input.Metadata.bInvertedAlpha;
 			Parameters.SourceEncoding = static_cast<uint32>(Input.Metadata.Encoding);
+			Parameters.Exposure = 1.0f;
+
+			if (CVarCompositeCoreFollowSceneExposure.GetValueOnRenderThread())
+			{
+				// When we output back to scene color, optionally adjust input exposure so that the ouput composite matches.
+				if (PassContext.bOutputSceneColor && !Input.Metadata.bPreExposed)
+				{
+					Parameters.Exposure = InPreExposure;
+				}
+			}
 
 			EDistortionUV DistortionUV = EDistortionUV::None;
 
@@ -124,6 +143,8 @@ namespace UE
 				}
 			}
 
+			const float PreExposure = InView.State ? InView.State->GetPreExposure() : 1.0f;
+
 			FScreenPassRenderTarget Output = Inputs.OverrideOutput;
 			if (!Output.IsValid())
 			{
@@ -136,8 +157,8 @@ namespace UE
 			PassParameters->Input0 = GetScreenPassTextureViewportParameters(FScreenPassTextureViewport(Inputs[0].Texture));
 			PassParameters->Input1 = GetScreenPassTextureViewportParameters(FScreenPassTextureViewport(Inputs[1].Texture));
 			PassParameters->Output = GetScreenPassTextureViewportParameters(FScreenPassTextureViewport(Output));
-			PassParameters->Tex0 = GetPassTextureParameters(Inputs[0], PassContext, Family, bUseLensDistortion);
-			PassParameters->Tex1 = GetPassTextureParameters(Inputs[1], PassContext, Family, bUseLensDistortion);
+			PassParameters->Tex0 = GetPassTextureParameters(Inputs[0], PassContext, Family, bUseLensDistortion, PreExposure);
+			PassParameters->Tex1 = GetPassTextureParameters(Inputs[1], PassContext, Family, bUseLensDistortion, PreExposure);
 			if (bUseLensDistortion)
 			{
 				PassParameters->DistortingDisplacementTexture = LensDistortionLUT.DistortingDisplacementTexture;
diff --git a/Engine/Plugins/Experimental/Compositing/CompositeCore/Source/CompositeCore/Private/Passes/CompositeCorePassProxy.cpp b/Engine/Plugins/Experimental/Compositing/CompositeCore/Source/CompositeCore/Private/Passes/CompositeCorePassProxy.cpp
--- a/Engine/Plugins/Experimental/Compositing/CompositeCore/Source/CompositeCore/Private/Passes/CompositeCorePassProxy.cpp
+++ b/Engine/Plugins/Experimental/Compositing/CompositeCore/Source/CompositeCore/Private/Passes/CompositeCorePassProxy.cpp
@@ -102,6 +102,7 @@ UE::CompositeCore::FPassInputArray::FPassInputArray(
 		FPassInput& Input = Inputs.AddDefaulted_GetRef();
 		Input.Texture = FScreenPassTexture::CopyFromSlice(GraphBuilder, InPostInputs.GetInput(EPostProcessMaterialInput::SceneColor));
 		Input.Metadata.bInvertedAlpha = true;
+		Input.Metadata.bPreExposed = true;
 		Input.Metadata.bDistorted = bLensDistortionEnabled && InLocation >= ISceneViewExtension::EPostProcessingPass::SSRInput;
 
 		// After-tonemap scene color may have encoding manually applied, as opposed to _SRGB textures
diff --git a/Engine/Plugins/Experimental/Compositing/CompositeCore/Source/CompositeCore/Public/CompositeCoreSettings.h b/Engine/Plugins/Experimental/Compositing/CompositeCore/Source/CompositeCore/Public/CompositeCoreSettings.h
--- a/Engine/Plugins/Experimental/Compositing/CompositeCore/Source/CompositeCore/Public/CompositeCoreSettings.h
+++ b/Engine/Plugins/Experimental/Compositing/CompositeCore/Source/CompositeCore/Public/CompositeCoreSettings.h
@@ -42,6 +42,9 @@ public:
 	virtual void PostInitProperties() override;
 
 public:
+	/** When enabled, the scene main render pre-exposure is applied onto separate composited render(s) (see CompositeCore.FollowSceneExposure cvar). */
+	UPROPERTY(config, EditAnywhere, Category = General, meta=(ConsoleVariable="CompositeCore.FollowSceneExposure"))
+	bool bCompositeFollowsSceneExposure;
 
 	// Ensure that the primitive class is allowed by checking against plugin settings.
 	COMPOSITECORE_API bool IsAllowedPrimitiveClass(const UPrimitiveComponent* InPrimitiveComponent) const;
diff --git a/Engine/Plugins/Experimental/Compositing/CompositeCore/Source/CompositeCore/Public/Passes/CompositeCorePassProxy.h b/Engine/Plugins/Experimental/Compositing/CompositeCore/Source/CompositeCore/Public/Passes/CompositeCorePassProxy.h
--- a/Engine/Plugins/Experimental/Compositing/CompositeCore/Source/CompositeCore/Public/Passes/CompositeCorePassProxy.h
+++ b/Engine/Plugins/Experimental/Compositing/CompositeCore/Source/CompositeCore/Public/Passes/CompositeCorePassProxy.h
@@ -53,6 +53,9 @@ namespace UE
 
 			/** Is the texture content distorted? */
 			bool bDistorted = false;
+
+			/** Is the texture's exposure already adjusted? */
+			bool bPreExposed = false;
 			
 			/** Source color encoding. */
 			EEncoding Encoding = EEncoding::Linear;
