CompositeCore: Remove non-general custom render pass unpremultiplication away from the general merge pass. Add support for opacification on the dilate pass.

* In MRQ, layers rendered separately with holdout masks are merged with an add operation. In our case, the main render is merged *over* the background, which must be opacified before the composite.
* The previous unpremultiplication step was insufficient: it yielded non opaque scene color alpha on translucent object sent through the basic holdout composite.

diff --git a/Engine/Plugins/Experimental/Compositing/CompositeCore/Shaders/Private/CompositeCoreDilate.usf b/Engine/Plugins/Experimental/Compositing/CompositeCore/Shaders/Private/CompositeCoreDilate.usf
--- a/Engine/Plugins/Experimental/Compositing/CompositeCore/Shaders/Private/CompositeCoreDilate.usf
+++ b/Engine/Plugins/Experimental/Compositing/CompositeCore/Shaders/Private/CompositeCoreDilate.usf
@@ -23,6 +23,7 @@
 Texture2D InputTexture;
 RWTexture2D<float4> RWOutputTexture;
 uint2 Dimensions;
+uint bOpacifyOutput;
 
 static const int2 kOffsets[KERNEL_COUNT] = {
     int2(-1,  0),
@@ -49,6 +50,20 @@ float4 SampleCachedColors(in uint2 GroupThreadId, int2 PixelOffset = int2(0,0))
     return SharedColor[GetTileArrayIndexFromPixelOffset(GroupThreadId, PixelOffset)];
 }
 
+float SafeRcp(float x)
+{
+	return x > 0.0 ? rcp(x) : 0.0;
+}
+
+float4 Opacify(float4 InColor)
+{
+    float RegularAlpha = saturate(1.0 - InColor.a);
+    // Unpremultiply color
+    InColor.xyz *= SafeRcp(RegularAlpha);
+    // Return opaque color, maintaining UE's inverted alpha convention
+	return float4(InColor.xyz, 0.0);
+}
+
 [numthreads(TILE_SIZE, TILE_SIZE, 1)]
 void MainCS(
     uint2 GroupId : SV_GroupID,
@@ -110,5 +125,5 @@ void MainCS(
         }
     }
 
-    RWOutputTexture[DispatchThreadId] = TexelColor;
+	RWOutputTexture[DispatchThreadId] = select(bOpacifyOutput > 0, Opacify(TexelColor), TexelColor);
 }
diff --git a/Engine/Plugins/Experimental/Compositing/CompositeCore/Shaders/Private/CompositeCoreMerge.usf b/Engine/Plugins/Experimental/Compositing/CompositeCore/Shaders/Private/CompositeCoreMerge.usf
--- a/Engine/Plugins/Experimental/Compositing/CompositeCore/Shaders/Private/CompositeCoreMerge.usf
+++ b/Engine/Plugins/Experimental/Compositing/CompositeCore/Shaders/Private/CompositeCoreMerge.usf
@@ -45,17 +45,7 @@ SamplerState UndistortingDisplacementSampler;
 float2 DisplayGamma; // DisplayGamma, 1/DisplayGamma
 uint MergeOp;
 uint OutputEncoding;
-uint bFlags; // 0: bInvertedAlphaOutput, 1: bUnpremulInput1
-
-float SafeRcp(float x)
-{
-	return x > 0.0 ? rcp(x) : 0.0;
-}
-
-float4 Unpremultiply(float4 InColor)
-{
-	return float4(InColor.xyz * SafeRcp(InColor.w), InColor.w);
-}
+uint bFlags; // 0: bInvertedAlphaOutput
 
 float3 Decode(float3 InColor, uint InEncoding)
 {
@@ -152,13 +142,6 @@ void MainPS(
 	Color0.a = GetOpacity(Color0.a, Tex0_bInvertedAlpha);
 	Color1.a = GetOpacity(Color1.a, Tex1_bInvertedAlpha);
 
-    {
-        // NOTE: This is necessary to extract the solid colors behind translucent alpha holdout masks.
-        // Could be moved elsewhere as a separate pre-processing step, with minor added cost.
-        const bool bUnpremulInput1 = ((bFlags >> 1) & 1);
-        Color1 = select(bUnpremulInput1, Unpremultiply(Color1), Color1);
-    }
-
     // Adjust exposure if necessary
     Color0.rgb *= Tex0_Exposure;
     Color1.rgb *= Tex1_Exposure;
diff --git a/Engine/Plugins/Experimental/Compositing/CompositeCore/Source/CompositeCore/Private/CompositeCoreSceneViewExtension.cpp b/Engine/Plugins/Experimental/Compositing/CompositeCore/Source/CompositeCore/Private/CompositeCoreSceneViewExtension.cpp
--- a/Engine/Plugins/Experimental/Compositing/CompositeCore/Source/CompositeCore/Private/CompositeCoreSceneViewExtension.cpp
+++ b/Engine/Plugins/Experimental/Compositing/CompositeCore/Source/CompositeCore/Private/CompositeCoreSceneViewExtension.cpp
@@ -26,12 +26,12 @@ class FCompositeCoreCustomRenderPass : public FCustomRenderPassBase
 public:
 	IMPLEMENT_CUSTOM_RENDER_PASS(FCompositeCoreCustomRenderPass);
 
-	FCompositeCoreCustomRenderPass(const FIntPoint& InRenderTargetSize, FCompositeCoreSceneViewExtension* InParentExtension, const FSceneView& InView, int32 InDilationSize)
+	FCompositeCoreCustomRenderPass(const FIntPoint& InRenderTargetSize, FCompositeCoreSceneViewExtension* InParentExtension, const FSceneView& InView, const UE::CompositeCore::FBuiltInRenderPassOptions& InOptions)
 		: FCustomRenderPassBase(TEXT("CompositeCoreCustomRenderPass"), ERenderMode::DepthAndBasePass, ERenderOutput::SceneColorAndAlpha, InRenderTargetSize)
 		, ParentExtension(InParentExtension)
 		, ViewId(InView.GetViewKey())
 		, ViewFeatureLevel(InView.GetFeatureLevel())
-		, DilationSize(InDilationSize)
+		, Inputs({ InOptions.DilationSize, InOptions.bOpacifyOutput })
 	{
 		bSceneColorWithTranslucent = true;
 	}
@@ -48,10 +48,9 @@ public:
 		const FRDGTextureDesc TextureDesc = FRDGTextureDesc::Create2D(RenderTargetTexture->Desc.Extent, PF_FloatRGBA, FClearValueBinding::Black, TexCreate_UAV | TexCreate_ShaderResource | TexCreate_RenderTargetable);
 		FRDGTextureRef Output = GraphBuilder.CreateTexture(TextureDesc, TEXT("CompositeCoreProcessedTexture"));
 
-		// TODO: Add support for wider dilation
-		if (DilationSize > 0)
+		if (Inputs.DilationSize > 0)
 		{
-			UE::CompositeCore::Private::AddDilatePass(GraphBuilder, RenderTargetTexture, Output, ViewFeatureLevel);
+			UE::CompositeCore::Private::AddDilatePass(GraphBuilder, RenderTargetTexture, Output, ViewFeatureLevel, Inputs);
 		}
 
 		ParentExtension->CollectCustomRenderTarget(ViewId, GraphBuilder.ConvertToExternalTexture(Output));
@@ -62,7 +61,7 @@ private:
 	FCompositeCoreSceneViewExtension* ParentExtension;
 	const uint32 ViewId;
 	const ERHIFeatureLevel::Type ViewFeatureLevel;
-	const int32 DilationSize;
+	const UE::CompositeCore::Private::FDilateInputs Inputs;
 };
 
 FCompositeCoreSceneViewExtension::FCompositeCoreSceneViewExtension(const FAutoRegister& AutoReg, UWorld* InWorld)
@@ -284,7 +283,7 @@ void FCompositeCoreSceneViewExtension::BeginRenderViewFamily(FSceneViewFamily& I
 			RenderTargetViewSize,
 			this,
 			InView,
-			BuiltInRenderPassOptions.OutputDilationSize
+			BuiltInRenderPassOptions
 		);
 
 		FSceneInterface::FCustomRenderPassRendererInput PassInput{};
diff --git a/Engine/Plugins/Experimental/Compositing/CompositeCore/Source/CompositeCore/Private/Passes/CompositeCorePassDilate.cpp b/Engine/Plugins/Experimental/Compositing/CompositeCore/Source/CompositeCore/Private/Passes/CompositeCorePassDilate.cpp
--- a/Engine/Plugins/Experimental/Compositing/CompositeCore/Source/CompositeCore/Private/Passes/CompositeCorePassDilate.cpp
+++ b/Engine/Plugins/Experimental/Compositing/CompositeCore/Source/CompositeCore/Private/Passes/CompositeCorePassDilate.cpp
@@ -18,6 +18,7 @@ class FCompositeCoreDilateShader : public FGlobalShader
 		SHADER_PARAMETER_RDG_TEXTURE(Texture2D, InputTexture)
 		SHADER_PARAMETER_RDG_TEXTURE_UAV(RWTexture2D<float4>, RWOutputTexture)
 		SHADER_PARAMETER(FIntPoint, Dimensions)
+		SHADER_PARAMETER(uint32, bOpacifyOutput)
 	END_SHADER_PARAMETER_STRUCT()
 
 	static inline void ModifyCompilationEnvironment(const FGlobalShaderPermutationParameters& Parameters, FShaderCompilerEnvironment& OutEnvironment)
@@ -37,7 +38,7 @@ namespace UE
 	{
 		namespace Private
 		{
-			void AddDilatePass(FRDGBuilder& GraphBuilder, FRDGTextureRef Input, FRDGTextureRef Output, ERHIFeatureLevel::Type FeatureLevel)
+			void AddDilatePass(FRDGBuilder& GraphBuilder, FRDGTextureRef Input, FRDGTextureRef Output, ERHIFeatureLevel::Type FeatureLevel, const FDilateInputs& PassInputs)
 			{
 				RDG_EVENT_SCOPE_STAT(GraphBuilder, FCompositeCoreDilate, "CompositeCore.Dilate");
 				RDG_GPU_STAT_SCOPE(GraphBuilder, FCompositeCoreDilate);
@@ -49,6 +50,7 @@ namespace UE
 				PassParameters->InputTexture = Input;
 				PassParameters->RWOutputTexture = GraphBuilder.CreateUAV(Output);
 				PassParameters->Dimensions = TextureSize;
+				PassParameters->bOpacifyOutput = PassInputs.bOpacifyOutput;
 
 				TShaderMapRef<FCompositeCoreDilateShader> ComputeShader(GlobalShaderMap);
 				FComputeShaderUtils::AddPass(
diff --git a/Engine/Plugins/Experimental/Compositing/CompositeCore/Source/CompositeCore/Private/Passes/CompositeCorePassDilate.h b/Engine/Plugins/Experimental/Compositing/CompositeCore/Source/CompositeCore/Private/Passes/CompositeCorePassDilate.h
--- a/Engine/Plugins/Experimental/Compositing/CompositeCore/Source/CompositeCore/Private/Passes/CompositeCorePassDilate.h
+++ b/Engine/Plugins/Experimental/Compositing/CompositeCore/Source/CompositeCore/Private/Passes/CompositeCorePassDilate.h
@@ -1,5 +1,7 @@
 // Copyright Epic Games, Inc. All Rights Reserved.
 
+#pragma once
+
 #include "RenderGraphFwd.h"
 #include "RHIFwd.h"
 
@@ -9,7 +11,21 @@ namespace UE
 	{
 		namespace Private
 		{
-			void AddDilatePass(FRDGBuilder& GraphBuilder, FRDGTextureRef Input, FRDGTextureRef Output, ERHIFeatureLevel::Type FeatureLevel);
+			struct FDilateInputs
+			{
+				/** Size of the dilation step. WARNING: Not currently implemented. */
+				int32 DilationSize = 1;
+				
+				/** Opacify the pass output to solid colors. */
+				bool bOpacifyOutput = true;
+			};
+
+			/**
+			* Compute shader dilation pass of non-translucent color texels. This is done as preparation
+			* for compositing to hide aliasing under the main render's anti-aliased edges, with an
+			* optional opacification step.
+			*/
+			void AddDilatePass(FRDGBuilder& GraphBuilder, FRDGTextureRef Input, FRDGTextureRef Output, ERHIFeatureLevel::Type FeatureLevel, const FDilateInputs& PassInputs = {});
 		}
 	}
 }
diff --git a/Engine/Plugins/Experimental/Compositing/CompositeCore/Source/CompositeCore/Private/Passes/CompositeCorePassMergeProxy.cpp b/Engine/Plugins/Experimental/Compositing/CompositeCore/Source/CompositeCore/Private/Passes/CompositeCorePassMergeProxy.cpp
--- a/Engine/Plugins/Experimental/Compositing/CompositeCore/Source/CompositeCore/Private/Passes/CompositeCorePassMergeProxy.cpp
+++ b/Engine/Plugins/Experimental/Compositing/CompositeCore/Source/CompositeCore/Private/Passes/CompositeCorePassMergeProxy.cpp
@@ -174,7 +174,7 @@ namespace UE
 			PassParameters->DisplayGamma = FVector2f(Family->RenderTarget->GetDisplayGamma(), 1.0f / Family->RenderTarget->GetDisplayGamma());
 			PassParameters->MergeOp = static_cast<uint32>(MergeOp);
 			PassParameters->OutputEncoding = static_cast<uint32>(OutMetadata.Encoding);
-			PassParameters->bFlags = static_cast<uint32>(bUnpremulInput1 << 1) | static_cast<uint32>(PassContext.bOutputSceneColor);
+			PassParameters->bFlags = static_cast<uint32>(PassContext.bOutputSceneColor);
 			PassParameters->RenderTargets[0] = Output.GetRenderTargetBinding();
 
 			FCompositeCoreMergeShader::FPermutationDomain PermutationVector;
diff --git a/Engine/Plugins/Experimental/Compositing/CompositeCore/Source/CompositeCore/Private/Passes/CompositeCorePassProxy.cpp b/Engine/Plugins/Experimental/Compositing/CompositeCore/Source/CompositeCore/Private/Passes/CompositeCorePassProxy.cpp
--- a/Engine/Plugins/Experimental/Compositing/CompositeCore/Source/CompositeCore/Private/Passes/CompositeCorePassProxy.cpp
+++ b/Engine/Plugins/Experimental/Compositing/CompositeCore/Source/CompositeCore/Private/Passes/CompositeCorePassProxy.cpp
@@ -22,7 +22,8 @@ namespace UE
 		FBuiltInRenderPassOptions::FBuiltInRenderPassOptions()
 			: ViewUserFlagsOverride{}
 			, bEnableUnlitViewmode{true}
-			, OutputDilationSize{1}
+			, DilationSize{1}
+			, bOpacifyOutput{true}
 		{
 		}
 
diff --git a/Engine/Plugins/Experimental/Compositing/CompositeCore/Source/CompositeCore/Public/Passes/CompositeCorePassMergeProxy.h b/Engine/Plugins/Experimental/Compositing/CompositeCore/Source/CompositeCore/Public/Passes/CompositeCorePassMergeProxy.h
--- a/Engine/Plugins/Experimental/Compositing/CompositeCore/Source/CompositeCore/Public/Passes/CompositeCorePassMergeProxy.h
+++ b/Engine/Plugins/Experimental/Compositing/CompositeCore/Source/CompositeCore/Public/Passes/CompositeCorePassMergeProxy.h
@@ -30,9 +30,6 @@ namespace UE
 
 			/** Merge operation. */
 			ECompositeCoreMergeOp MergeOp = ECompositeCoreMergeOp::Over;
-
-			/** Unpremultiplication required for recovering opaque colors for the default over blend. */
-			bool bUnpremulInput1 = true;
 		};
 	}
 }
diff --git a/Engine/Plugins/Experimental/Compositing/CompositeCore/Source/CompositeCore/Public/Passes/CompositeCorePassProxy.h b/Engine/Plugins/Experimental/Compositing/CompositeCore/Source/CompositeCore/Public/Passes/CompositeCorePassProxy.h
--- a/Engine/Plugins/Experimental/Compositing/CompositeCore/Source/CompositeCore/Public/Passes/CompositeCorePassProxy.h
+++ b/Engine/Plugins/Experimental/Compositing/CompositeCore/Source/CompositeCore/Public/Passes/CompositeCorePassProxy.h
@@ -209,7 +209,10 @@ namespace UE
 			bool bEnableUnlitViewmode;
 
 			/** Enables the dilation pass with size > 0 (currently only 1 is supported). */
-			int32 OutputDilationSize;
+			int32 DilationSize;
+
+			/** Opacify to extract the solid colors behind translucent alpha holdout masks. */
+			bool bOpacifyOutput;
 		};
 
 		/** Render-thread struct for render post-processing work per frame. */
