HoldoutComposite: Fix asymmetrical unregister/register primitive calls for components inside blueprint objects.

Blueprint objects have bRunConstructionScriptOnDrag enabled by default, which results in a component unregister call via AActor::PostEditMove when the actove is moved in editor. Problem was that re-registration was only done on the next frame tick.

diff --git a/Engine/Plugins/Experimental/Compositing/CompositeCore/Source/CompositeCore/Private/HoldoutCompositeComponent.cpp b/Engine/Plugins/Experimental/Compositing/CompositeCore/Source/CompositeCore/Private/HoldoutCompositeComponent.cpp
--- a/Engine/Plugins/Experimental/Compositing/CompositeCore/Source/CompositeCore/Private/HoldoutCompositeComponent.cpp
+++ b/Engine/Plugins/Experimental/Compositing/CompositeCore/Source/CompositeCore/Private/HoldoutCompositeComponent.cpp
@@ -74,16 +74,30 @@ UHoldoutCompositeComponent::UHoldoutCompositeComponent(const FObjectInitializer&
 	}
 }
 
+void UHoldoutCompositeComponent::BeginDestroy()
+{
+	Super::BeginDestroy();
+
+	// Kept for safety, but redundant since the scene view extension will automatically discard invalid primitive objects.
+	UnregisterCompositeImpl();
+}
+
 void UHoldoutCompositeComponent::OnRegister()
 {
 	Super::OnRegister();
 
-	RegisterCompositeImpl();
+	if (!UE::CompositeCore::IsRegisterPrimitivesOnTickEnabled())
+	{
+		RegisterCompositeImpl();
+	}
 }
 
 void UHoldoutCompositeComponent::OnUnregister()
 {
-	UnregisterCompositeImpl();
+	if (!UE::CompositeCore::IsRegisterPrimitivesOnTickEnabled())
+	{
+		UnregisterCompositeImpl();
+	}
 
 	Super::OnUnregister();
 }
@@ -92,8 +106,10 @@ void UHoldoutCompositeComponent::TickComponent(float DeltaTime, ELevelTick TickT
 {
 	Super::TickComponent(DeltaTime, TickType, ThisTickFunction);
 
-	constexpr bool bIsTick = true;
-	RegisterCompositeImpl(bIsTick);
+	if (UE::CompositeCore::IsRegisterPrimitivesOnTickEnabled())
+	{
+		RegisterCompositeImpl();
+	}
 }
 
 void UHoldoutCompositeComponent::DetachFromComponent(const FDetachmentTransformRules& DetachmentRules)
@@ -153,15 +169,10 @@ void UHoldoutCompositeComponent::SetEnabled(bool bInIsEnabled)
 	}
 }
 
-void UHoldoutCompositeComponent::RegisterCompositeImpl(bool bIsTick)
+void UHoldoutCompositeComponent::RegisterCompositeImpl()
 {
 	using namespace UE::CompositeCore;
 
-	if (IsRegisterPrimitivesOnTickEnabled() != bIsTick)
-	{
-		return;
-	}
-
 	if (!bIsEnabled)
 	{
 		return;
diff --git a/Engine/Plugins/Experimental/Compositing/CompositeCore/Source/CompositeCore/Public/HoldoutCompositeComponent.h b/Engine/Plugins/Experimental/Compositing/CompositeCore/Source/CompositeCore/Public/HoldoutCompositeComponent.h
--- a/Engine/Plugins/Experimental/Compositing/CompositeCore/Source/CompositeCore/Public/HoldoutCompositeComponent.h
+++ b/Engine/Plugins/Experimental/Compositing/CompositeCore/Source/CompositeCore/Public/HoldoutCompositeComponent.h
@@ -13,6 +13,10 @@ class UHoldoutCompositeComponent : public USceneComponent
 	GENERATED_UCLASS_BODY()
 
 public:
+	//~ Begin UActorComponent interface
+	virtual void BeginDestroy() override;
+	//~ End UActorComponent interface
+
 	//~ Begin UActorComponent interface
 	virtual void OnRegister() override;
 	virtual void OnUnregister() override;
@@ -35,7 +39,7 @@ public:
 private:
 
 	/* Private implementation of the register method. */
-	void RegisterCompositeImpl(bool bIsTick = false);
+	void RegisterCompositeImpl();
 
 	/* Private implementation of the unregister method. */
 	void UnregisterCompositeImpl();
