CompositeCore: Hotfix for alpha issue introduced by dilation opacification CL 41646203.

diff --git a/Engine/Plugins/Experimental/Compositing/CompositeCore/Shaders/Private/CompositeCoreDilate.usf b/Engine/Plugins/Experimental/Compositing/CompositeCore/Shaders/Private/CompositeCoreDilate.usf
--- a/Engine/Plugins/Experimental/Compositing/CompositeCore/Shaders/Private/CompositeCoreDilate.usf
+++ b/Engine/Plugins/Experimental/Compositing/CompositeCore/Shaders/Private/CompositeCoreDilate.usf
@@ -7,7 +7,7 @@
 #endif
 
 #define KERNEL_COUNT 8
-#define ALPHA_THRESHOLD (1.0f - 1e-7f)
+#define NEARLY_TRANSLUCENT_ALPHA 0.9999999f // Alpha is inverted opacity
 
 // Number of texels arround the group tile for scene color.
 #define LDS_TILE_BORDER_SIZE (1)
@@ -110,8 +110,9 @@ void MainCS(
 
     float4 TexelColor = SampleCachedColors(GroupThreadId);
     
+	// Apply (gather) dilation on any translucent or nearly translucent texel
     BRANCH
-    if(TexelColor.a > ALPHA_THRESHOLD) // Alpha is inverted opacity
+    if(TexelColor.a > NEARLY_TRANSLUCENT_ALPHA)
     {
         UNROLL
         for(int i=0; i<KERNEL_COUNT; i++)
@@ -125,5 +126,11 @@ void MainCS(
         }
     }
 
-	RWOutputTexture[DispatchThreadId] = select(bOpacifyOutput > 0, Opacify(TexelColor), TexelColor);
+	// Condtionally & fully opacify non-translucent colors for the over blend.
+	if (bOpacifyOutput && TexelColor.a < NEARLY_TRANSLUCENT_ALPHA)
+	{
+		TexelColor = Opacify(TexelColor);
+	}
+
+	RWOutputTexture[DispatchThreadId] = TexelColor;
 }
