CompositeCore: Add convenient support for FXAA on the default custom render pass with new "Apply FXAA" plugin setting.

* Additional fix for FXAA pass where we apply a simple invertible tonecurve to colors for FXAA to operate in display/perceptual space.
* Additional fix for cvar to project setting config property not loaded correctly on launch.
* Additional fix for potentially incorrectly-sized custom render pass resolution causing softness with bilinear filtering.
* Renamed pre-exposure setting.

diff --git a/Engine/Plugins/Experimental/Compositing/CompositeCore/Config/BaseCompositeCore.ini b/Engine/Plugins/Experimental/Compositing/CompositeCore/Config/BaseCompositeCore.ini
--- a/Engine/Plugins/Experimental/Compositing/CompositeCore/Config/BaseCompositeCore.ini
+++ b/Engine/Plugins/Experimental/Compositing/CompositeCore/Config/BaseCompositeCore.ini
@@ -2,4 +2,6 @@
 
 +ClassRedirects=(OldName="/Script/HoldoutComposite.HoldoutCompositeComponent", NewName="/Script/CompositeCore.HoldoutCompositeComponent")
 +ClassRedirects=(OldName="/Script/HoldoutComposite.HoldoutCompositeSubsystem", NewName="/Script/CompositeCore.CompositeCoreSubsystem")
-+ClassRedirects=(OldName="/Script/HoldoutComposite.HoldoutCompositeSettings", NewName="/Script/CompositeCore.CompositeCoreSettings")
\ No newline at end of file
++ClassRedirects=(OldName="/Script/HoldoutComposite.HoldoutCompositeSettings", NewName="/Script/CompositeCore.CompositeCoreSettings")
+
++PropertyRedirects=(OldName="/Script/HoldoutComposite.HoldoutCompositeSettings.bCompositeFollowsSceneExposure", NewName="/Script/CompositeCore.CompositeCoreSettings.bApplyPreExposure")
\ No newline at end of file
diff --git a/Engine/Plugins/Experimental/Compositing/CompositeCore/Shaders/Private/CompositeDisplayTransform.usf b/Engine/Plugins/Experimental/Compositing/CompositeCore/Shaders/Private/CompositeDisplayTransform.usf
new file mode 100644
--- /dev/null
+++ b/Engine/Plugins/Experimental/Compositing/CompositeCore/Shaders/Private/CompositeDisplayTransform.usf
@@ -0,0 +1,59 @@
+// Copyright Epic Games, Inc. All Rights Reserved.
+
+#include "/Engine/Private/Common.ush"
+
+Texture2D InputTexture;
+SamplerState InputSampler;
+uint bIsForward;
+float Gamma;
+float InvGamma;
+
+#define VERSION 0
+#define R_K 1.0
+
+// Simple Reinhard tone curve
+float3 Reinhard(float3 ColorHDR, float K = R_K)
+{
+    return ColorHDR / (ColorHDR + K);
+}
+
+// Inverse simple Reinhard tone curve
+float3 InvReinhard(float3 ColorSDR, float K = R_K)
+{
+    return K * ColorSDR / saturate(1.0 - ColorSDR);
+}
+
+// Gamma encode (faster approximation of sRGB)
+float3 GammaEncode(float3 Color)
+{
+    return pow(Color, InvGamma);
+}
+
+// Gamma decode transfer function (faster approximation of sRGB)
+float3 GammaDecode(float3 Color)
+{
+    return pow(Color, Gamma);
+}
+
+void MainPS(
+    noperspective float4 InUVAndScreenPos : TEXCOORD0,
+    out float4 OutColor : SV_Target0
+    )
+{
+    OutColor = InputTexture.Sample(InputSampler, InUVAndScreenPos.xy);
+
+#if VERSION == 0
+    if(bIsForward)
+    {
+        OutColor.rgb = GammaEncode(Reinhard(OutColor.rgb));
+    }
+    else
+    {
+        OutColor.rgb = InvReinhard(GammaDecode(OutColor.rgb));
+    }
+#elif VERSION == 1
+    OutColor.rgb = bIsForward ? Reinhard(OutColor.rgb) : InvReinhard(OutColor.rgb);
+#elif VERSION == 2
+    OutColor.rgb = bIsForward ? GammaEncode(OutColor.rgb) : GammaDecode(OutColor.rgb);
+#endif
+}
diff --git a/Engine/Plugins/Experimental/Compositing/CompositeCore/Source/CompositeCore/Private/CompositeCoreModule.cpp b/Engine/Plugins/Experimental/Compositing/CompositeCore/Source/CompositeCore/Private/CompositeCoreModule.cpp
--- a/Engine/Plugins/Experimental/Compositing/CompositeCore/Source/CompositeCore/Private/CompositeCoreModule.cpp
+++ b/Engine/Plugins/Experimental/Compositing/CompositeCore/Source/CompositeCore/Private/CompositeCoreModule.cpp
@@ -4,6 +4,7 @@
 
 #include "Interfaces/IPluginManager.h"
 #include "HDRHelper.h"
+#include "Misc/ConfigUtilities.h"
 #include "Misc/Paths.h"
 #include "Modules/ModuleManager.h"
 #include "ShaderCore.h"
@@ -16,6 +17,9 @@ void FCompositeCoreModule::StartupModule()
 {
 	AddShaderSourceDirectoryMapping(TEXT("/Plugin/CompositeCore"), FPaths::Combine(IPluginManager::Get().FindPlugin(TEXT("CompositeCore"))->GetBaseDir(), TEXT("Shaders")));
 
+	// Since we are so early in the loading phase, we first need to load the cvars since they're not loaded at this point.
+	UE::ConfigUtilities::ApplyCVarSettingsFromIni(TEXT("/Script/CompositeCore.CompositeCorePluginSettings"), *GEngineIni, ECVF_SetByProjectSetting);
+
 	if (IsHDREnabled())
 	{
 		UE_LOG(LogCompositeCore, Warning, TEXT("Composite pipeline disabled: HDR mode is not currently supported."));
diff --git a/Engine/Plugins/Experimental/Compositing/CompositeCore/Source/CompositeCore/Private/CompositeCoreSceneViewExtension.cpp b/Engine/Plugins/Experimental/Compositing/CompositeCore/Source/CompositeCore/Private/CompositeCoreSceneViewExtension.cpp
--- a/Engine/Plugins/Experimental/Compositing/CompositeCore/Source/CompositeCore/Private/CompositeCoreSceneViewExtension.cpp
+++ b/Engine/Plugins/Experimental/Compositing/CompositeCore/Source/CompositeCore/Private/CompositeCoreSceneViewExtension.cpp
@@ -4,6 +4,7 @@
 
 #include "CompositeCoreModule.h"
 #include "Passes/CompositeCorePassDilate.h"
+#include "Passes/CompositeCorePassFXAAProxy.h"
 
 #include "CommonRenderResources.h"
 #include "Components/PrimitiveComponent.h"
@@ -12,6 +13,7 @@
 #include "Engine/Texture.h"
 #include "HDRHelper.h"
 #include "PostProcess/LensDistortion.h"
+#include "PostProcess/PostProcessAA.h"
 #include "PostProcess/PostProcessMaterialInputs.h"
 #include "Rendering/CustomRenderPass.h"
 #include "RenderGraphBuilder.h"
@@ -21,6 +23,12 @@
 #include "SystemTextures.h"
 #include "TextureResource.h"
 
+static TAutoConsoleVariable<int32> CVarCompositeCoreApplyFXAA(
+	TEXT("CompositeCore.ApplyFXAA"),
+	0,
+	TEXT("When enabled, the custom render pass automatically applies FXAA."),
+	ECVF_RenderThreadSafe);
+
 class FCompositeCoreCustomRenderPass : public FCustomRenderPassBase
 {
 public:
@@ -262,25 +270,9 @@ void FCompositeCoreSceneViewExtension::BeginRenderViewFamily(FSceneViewFamily& I
 			return;
 		}
 
-		// Extract the custom render target size
-		FIntPoint RenderTargetViewSize;
-		if (InView.PrimaryScreenPercentageMethod == EPrimaryScreenPercentageMethod::TemporalUpscale)
-		{
-			// Note: This is equivalent to FViewInfo::GetSecondaryViewRectSize
-			FIntRect OutputRect = FIntRect(0, 0,
-				FMath::CeilToInt(InView.UnscaledViewRect.Width() * InViewFamily.SecondaryViewFraction),
-				FMath::CeilToInt(InView.UnscaledViewRect.Height() * InViewFamily.SecondaryViewFraction)
-			);
-			QuantizeSceneBufferSize(OutputRect.Max, RenderTargetViewSize);
-		}
-		else
-		{
-			RenderTargetViewSize = InView.UnscaledViewRect.Size();
-		}
-
 		// Create a new custom render pass to render the composite primitive(s)
 		FCompositeCoreCustomRenderPass* CustomRenderPass = new FCompositeCoreCustomRenderPass(
-			RenderTargetViewSize,
+			InView.UnscaledViewRect.Size(),
 			this,
 			InView,
 			BuiltInRenderPassOptions
@@ -323,6 +315,26 @@ void FCompositeCoreSceneViewExtension::PreRenderView_RenderThread(FRDGBuilder& G
 
 void FCompositeCoreSceneViewExtension::PrePostProcessPass_RenderThread(FRDGBuilder& GraphBuilder, const FSceneView& InView, const FPostProcessingInputs& Inputs)
 {
+	using namespace UE::CompositeCore;
+
+	TRefCountPtr<IPooledRenderTarget>* CompositeRenderPassPtr = CustomRenderTargetPerView_RenderThread.Find(InView.GetViewKey());
+
+	if ((CompositeRenderPassPtr != nullptr) && CVarCompositeCoreApplyFXAA.GetValueOnRenderThread())
+	{
+		static FFXAAPassProxy FXAAPassProxy = FFXAAPassProxy(DefaultPassInputDecl);
+
+		// Set the composite render target as the pass input
+		FPassInputArray PassInputs;
+		FPassInput& PassInput = PassInputs.GetArray().AddDefaulted_GetRef();
+		PassInput.Texture = FScreenPassTexture{ GraphBuilder.RegisterExternalTexture(*CompositeRenderPassPtr) };
+
+		// Apply FXAA, with additional fwd/inv display transform passes
+		const FPassOutput Output = FXAAPassProxy.Add(GraphBuilder, InView, PassInputs, {});
+		const FScreenPassTexture& ScreenTexOutput = Output.Resource.Texture;
+
+		// Extract the result back into the composite render target
+		*CompositeRenderPassPtr = GraphBuilder.ConvertToExternalTexture(ScreenTexOutput.Texture);
+	}
 }
 
 void FCompositeCoreSceneViewExtension::SubscribeToPostProcessingPass(EPostProcessingPass PassId, const FSceneView& InView, FAfterPassCallbackDelegateArray& InOutPassCallbacks, bool bIsPassEnabled)
diff --git a/Engine/Plugins/Experimental/Compositing/CompositeCore/Source/CompositeCore/Private/CompositeCoreSettings.cpp b/Engine/Plugins/Experimental/Compositing/CompositeCore/Source/CompositeCore/Private/CompositeCoreSettings.cpp
--- a/Engine/Plugins/Experimental/Compositing/CompositeCore/Source/CompositeCore/Private/CompositeCoreSettings.cpp
+++ b/Engine/Plugins/Experimental/Compositing/CompositeCore/Source/CompositeCore/Private/CompositeCoreSettings.cpp
@@ -24,7 +24,8 @@ static FAutoConsoleVariableRef CVarRefreshPrimitivesOnTick(
 	ECVF_ReadOnly);
 
 UCompositeCorePluginSettings::UCompositeCorePluginSettings()
-	: bCompositeFollowsSceneExposure(false)
+	: bApplyPreExposure(0)
+	, bApplyFXAA(0)
 	, DisabledPrimitiveClasses(
 		{
 			FSoftClassPath("/Script/Engine.CapsuleComponent"),
@@ -53,6 +54,7 @@ FName UCompositeCorePluginSettings::GetSectionName() const
 {
 	return TEXT("Composite Core");
 }
+
 void UCompositeCorePluginSettings::PostEditChangeProperty(FPropertyChangedEvent& PropertyChangedEvent)
 {
 	Super::PostEditChangeProperty(PropertyChangedEvent);
@@ -62,7 +64,10 @@ void UCompositeCorePluginSettings::PostEditChangeProperty(FPropertyChangedEvent&
 		CacheDisabledClasses();
 	}
 
-	ExportValuesToConsoleVariables(PropertyChangedEvent.Property);
+	if (PropertyChangedEvent.Property)
+	{
+		ExportValuesToConsoleVariables(PropertyChangedEvent.Property);
+	}
 }
 #endif
 
@@ -70,6 +75,13 @@ void UCompositeCorePluginSettings::PostInitProperties()
 {
 	Super::PostInitProperties();
 
+#if WITH_EDITOR
+	if (IsTemplate())
+	{
+		ImportConsoleVariableValues();
+	}
+#endif // #if WITH_EDITOR
+
 	CacheDisabledClasses();
 }
 
diff --git a/Engine/Plugins/Experimental/Compositing/CompositeCore/Source/CompositeCore/Private/Passes/CompositeCorePassFXAAProxy.cpp b/Engine/Plugins/Experimental/Compositing/CompositeCore/Source/CompositeCore/Private/Passes/CompositeCorePassFXAAProxy.cpp
new file mode 100644
--- /dev/null
+++ b/Engine/Plugins/Experimental/Compositing/CompositeCore/Source/CompositeCore/Private/Passes/CompositeCorePassFXAAProxy.cpp
@@ -0,0 +1,105 @@
+// Copyright Epic Games, Inc. All Rights Reserved.
+
+#include "Passes/CompositeCorePassFXAAProxy.h"
+
+#include "RenderGraphBuilder.h"
+#include "RenderGraphUtils.h"
+#include "SceneView.h"
+#include "ScreenPass.h"
+#include "PostProcess/PostProcessAA.h"
+#include "PostProcess/PostProcessMaterialInputs.h"
+
+DECLARE_GPU_STAT_NAMED(FCompositeCoreDisplayTransform, TEXT("CompositeCore.DisplayTransform"));
+
+class FCompositeCoreDisplayTransformShader : public FGlobalShader
+{
+	DECLARE_GLOBAL_SHADER(FCompositeCoreDisplayTransformShader);
+	SHADER_USE_PARAMETER_STRUCT(FCompositeCoreDisplayTransformShader, FGlobalShader);
+
+	BEGIN_SHADER_PARAMETER_STRUCT(FParameters, )
+		SHADER_PARAMETER_RDG_TEXTURE(Texture2D, InputTexture)
+		SHADER_PARAMETER_SAMPLER(SamplerState, InputSampler)
+		SHADER_PARAMETER(FIntPoint, Dimensions)
+		SHADER_PARAMETER(uint32, bIsForward)
+		SHADER_PARAMETER(float, Gamma)
+		SHADER_PARAMETER(float, InvGamma)
+		RENDER_TARGET_BINDING_SLOTS()
+	END_SHADER_PARAMETER_STRUCT()
+};
+IMPLEMENT_GLOBAL_SHADER(FCompositeCoreDisplayTransformShader, "/Plugin/CompositeCore/Private/CompositeDisplayTransform.usf", "MainPS", SF_Pixel);
+
+namespace UE
+{
+	namespace CompositeCore
+	{
+		namespace Private
+		{
+			FScreenPassTexture AddDisplayTransformPass(FRDGBuilder& GraphBuilder, const FScreenPassViewInfo ViewInfo, const FScreenPassTexture& Input, bool bIsForward, float Gamma = 2.2f)
+			{
+				RDG_EVENT_SCOPE_STAT(GraphBuilder, FCompositeCoreDisplayTransform, "CompositeCore.DisplayTransform");
+				RDG_GPU_STAT_SCOPE(GraphBuilder, FCompositeCoreDisplayTransform);
+
+				FGlobalShaderMap* GlobalShaderMap = GetGlobalShaderMap(ViewInfo.FeatureLevel);
+				FScreenPassRenderTarget Output = FScreenPassRenderTarget::CreateFromInput(GraphBuilder, Input, ERenderTargetLoadAction::ENoAction, TEXT("FXAACompositePass"));
+				FScreenPassTextureViewport Viewport = FScreenPassTextureViewport(Output);
+
+				FCompositeCoreDisplayTransformShader::FParameters* Parameters = GraphBuilder.AllocParameters<FCompositeCoreDisplayTransformShader::FParameters>();
+				Parameters->InputTexture = Input.Texture;
+				Parameters->InputSampler = TStaticSamplerState<>::GetRHI();
+				Parameters->bIsForward = bIsForward;
+				Parameters->Gamma = Gamma;
+				Parameters->InvGamma = 1.0f / Gamma;
+				Parameters->bIsForward = bIsForward;
+				Parameters->RenderTargets[0] = Output.GetRenderTargetBinding();
+
+				TShaderMapRef<FCompositeCoreDisplayTransformShader> PixelShader(GlobalShaderMap);
+				AddDrawScreenPass(
+					GraphBuilder,
+					RDG_EVENT_NAME("Composite.DisplayTransform (%dx%d)", Viewport.Extent.X, Viewport.Extent.Y),
+					ViewInfo,
+					Viewport,
+					Viewport,
+					PixelShader,
+					Parameters
+				);
+
+				return Output;
+			}
+		}
+
+		FPassOutput FFXAAPassProxy::Add(FRDGBuilder& GraphBuilder, const FSceneView& InView, const FPassInputArray& Inputs, const FPassContext& PassContext) const
+		{
+			FScreenPassTexture Input = Inputs[0].Texture;
+			FResourceMetadata Metadata = Inputs[0].Metadata;
+			const bool bLinearSourceColors = (Metadata.Encoding == EEncoding::Linear);
+
+			if (bLinearSourceColors)
+			{
+				// We tonemap & encode the result so that FXAA can operate on perceptual colors
+				constexpr bool bIsForward = true;
+				Input = Private::AddDisplayTransformPass(GraphBuilder, InView, Input, bIsForward);
+			}
+
+			static IConsoleVariable* FXAAQualityCVar = IConsoleManager::Get().FindConsoleVariable(TEXT("r.FXAA.Quality"));
+			check(FXAAQualityCVar);
+			const int32 Quality = QualityOverride.IsSet() ? QualityOverride.GetValue() : FXAAQualityCVar->GetInt();
+
+			FFXAAInputs FXAAInputs{};
+			FXAAInputs.SceneColor = Input;
+			FXAAInputs.Quality = static_cast<EFXAAQuality>(FMath::Clamp(Quality, 0, static_cast<int32>(EFXAAQuality::MAX) - 1));
+			FXAAInputs.OverrideOutput = Inputs.OverrideOutput;
+
+			FScreenPassTexture Output = AddFXAAPass(GraphBuilder, InView, FXAAInputs);
+
+			if (bLinearSourceColors)
+			{
+				// We decode and invert the tonemapping to obtain linear colors again.
+				constexpr bool bIsForward = false;
+				Output = Private::AddDisplayTransformPass(GraphBuilder, InView, Output, bIsForward);
+			}
+
+			return FPassOutput{ MoveTemp(Output), Metadata, PassOutputOverride };
+		}
+	}
+}
+
diff --git a/Engine/Plugins/Experimental/Compositing/CompositeCore/Source/CompositeCore/Private/Passes/CompositeCorePassMergeProxy.cpp b/Engine/Plugins/Experimental/Compositing/CompositeCore/Source/CompositeCore/Private/Passes/CompositeCorePassMergeProxy.cpp
--- a/Engine/Plugins/Experimental/Compositing/CompositeCore/Source/CompositeCore/Private/Passes/CompositeCorePassMergeProxy.cpp
+++ b/Engine/Plugins/Experimental/Compositing/CompositeCore/Source/CompositeCore/Private/Passes/CompositeCorePassMergeProxy.cpp
@@ -15,9 +15,9 @@
 
 DECLARE_GPU_STAT_NAMED(FCompositeCoreMerge, TEXT("CompositeCore.Merge"));
 
-static TAutoConsoleVariable<bool> CVarCompositeCoreFollowSceneExposure(
-	TEXT("CompositeCore.FollowSceneExposure"),
-	false,
+static TAutoConsoleVariable<int32> CVarCompositeCoreApplyPreExposure(
+	TEXT("CompositeCore.ApplyPreExposure"),
+	0,
 	TEXT("When enabled, the scene main render pre-exposure is applied onto separate composited render(s)."),
 	ECVF_RenderThreadSafe);
 
@@ -84,7 +84,7 @@ namespace UE
 			Parameters.SourceEncoding = static_cast<uint32>(Input.Metadata.Encoding);
 			Parameters.Exposure = 1.0f;
 
-			if (CVarCompositeCoreFollowSceneExposure.GetValueOnRenderThread())
+			if (CVarCompositeCoreApplyPreExposure.GetValueOnRenderThread())
 			{
 				// When we output back to scene color, optionally adjust input exposure so that the ouput composite matches.
 				if (PassContext.bOutputSceneColor && !Input.Metadata.bPreExposed)
diff --git a/Engine/Plugins/Experimental/Compositing/CompositeCore/Source/CompositeCore/Private/Passes/CompositeCorePassProxy.cpp b/Engine/Plugins/Experimental/Compositing/CompositeCore/Source/CompositeCore/Private/Passes/CompositeCorePassProxy.cpp
--- a/Engine/Plugins/Experimental/Compositing/CompositeCore/Source/CompositeCore/Private/Passes/CompositeCorePassProxy.cpp
+++ b/Engine/Plugins/Experimental/Compositing/CompositeCore/Source/CompositeCore/Private/Passes/CompositeCorePassProxy.cpp
@@ -62,6 +62,12 @@ namespace UE
 }
 
 
+FCompositeCorePassProxy::FCompositeCorePassProxy(UE::CompositeCore::EDefaultPassInputDecl)
+	: PassDeclaredInputs({ UE::CompositeCore::FPassInputDecl(TInPlaceType<UE::CompositeCore::FPassInternalResourceDesc>(), UE::CompositeCore::FPassInternalResourceDesc{}) })
+	, PassOutputOverride()
+{
+}
+
 FCompositeCorePassProxy::FCompositeCorePassProxy(UE::CompositeCore::FPassInputDeclArray InPassDeclaredInputs, TOptional<UE::CompositeCore::ResourceId> InPassOutputOverride)
 	: PassDeclaredInputs(MoveTemp(InPassDeclaredInputs))
 	, PassOutputOverride(MoveTemp(InPassOutputOverride))
diff --git a/Engine/Plugins/Experimental/Compositing/CompositeCore/Source/CompositeCore/Public/CompositeCoreSettings.h b/Engine/Plugins/Experimental/Compositing/CompositeCore/Source/CompositeCore/Public/CompositeCoreSettings.h
--- a/Engine/Plugins/Experimental/Compositing/CompositeCore/Source/CompositeCore/Public/CompositeCoreSettings.h
+++ b/Engine/Plugins/Experimental/Compositing/CompositeCore/Source/CompositeCore/Public/CompositeCoreSettings.h
@@ -22,7 +22,7 @@ namespace UE
 /**
  * Settings for the CompositeCore module.
  */
-UCLASS(MinimalAPI, config = Engine, meta = (DisplayName = "Composite Core"), MinimalAPI)
+UCLASS(MinimalAPI, config = Engine, defaultconfig, meta = (DisplayName = "Composite Core"))
 class UCompositeCorePluginSettings : public UDeveloperSettings
 {
 	GENERATED_BODY()
@@ -39,16 +39,26 @@ public:
 #endif
 	//~ End UDeveloperSettings interface
 
+	//~ Begin UObject Interface
 	virtual void PostInitProperties() override;
-
-public:
-	/** When enabled, the scene main render pre-exposure is applied onto separate composited render(s) (see CompositeCore.FollowSceneExposure cvar). */
-	UPROPERTY(config, EditAnywhere, Category = General, meta=(ConsoleVariable="CompositeCore.FollowSceneExposure"))
-	bool bCompositeFollowsSceneExposure;
+	//~ End UObject Interface
 
 	// Ensure that the primitive class is allowed by checking against plugin settings.
 	COMPOSITECORE_API bool IsAllowedPrimitiveClass(const UPrimitiveComponent* InPrimitiveComponent) const;
 
+public:
+	UPROPERTY(config, EditAnywhere, Category = CompositeCore, meta = (
+		ConsoleVariable = "CompositeCore.ApplyPreExposure", DisplayName = "Apply Pre-Exposure",
+		ToolTip = "When enabled, the scene main render pre-exposure is applied onto the separate composited render. This can be used to match exposure to the scene. (Maps to \"CompositeCore.ApplyPreExposure\" console variable).",
+		ConfigRestartRequired = false))
+	uint32 bApplyPreExposure : 1;
+
+	UPROPERTY(config, EditAnywhere, Category = CompositeCore, meta = (
+		ConsoleVariable = "CompositeCore.ApplyFXAA", DisplayName = "Apply FXAA",
+		ToolTip = "When enabled, FXAA is applied onto the separate composited render. Quality is controlled with \"r.FXAA.Quality\". (Maps to \"CompositeCore.ApplyFXAA\" console variable).",
+		ConfigRestartRequired = false))
+	uint32 bApplyFXAA : 1;
+
 	/** Primitive component classes that do not support the composite pipeline.*/
 	UPROPERTY(config, EditAnywhere, Category = CompositeCore)
 	TArray<FSoftClassPath> DisabledPrimitiveClasses;
diff --git a/Engine/Plugins/Experimental/Compositing/CompositeCore/Source/CompositeCore/Public/Passes/CompositeCorePassFXAAProxy.h b/Engine/Plugins/Experimental/Compositing/CompositeCore/Source/CompositeCore/Public/Passes/CompositeCorePassFXAAProxy.h
new file mode 100644
--- /dev/null
+++ b/Engine/Plugins/Experimental/Compositing/CompositeCore/Source/CompositeCore/Public/Passes/CompositeCorePassFXAAProxy.h
@@ -0,0 +1,23 @@
+// Copyright Epic Games, Inc. All Rights Reserved.
+
+#pragma once
+
+#include "CoreTypes.h"
+#include "CompositeCorePassProxy.h"
+
+namespace UE
+{
+	namespace CompositeCore
+	{
+		class COMPOSITECORE_API FFXAAPassProxy : public FCompositeCorePassProxy
+		{
+		public:
+			using FCompositeCorePassProxy::FCompositeCorePassProxy;
+
+			FPassOutput Add(FRDGBuilder& GraphBuilder, const FSceneView& InView, const FPassInputArray& Inputs, const FPassContext& PassContext) const override;
+
+			/** Optional r.FXAA.Quality setting override.  */
+			TOptional<int32> QualityOverride = {};
+		};
+	}
+}
diff --git a/Engine/Plugins/Experimental/Compositing/CompositeCore/Source/CompositeCore/Public/Passes/CompositeCorePassProxy.h b/Engine/Plugins/Experimental/Compositing/CompositeCore/Source/CompositeCore/Public/Passes/CompositeCorePassProxy.h
--- a/Engine/Plugins/Experimental/Compositing/CompositeCore/Source/CompositeCore/Public/Passes/CompositeCorePassProxy.h
+++ b/Engine/Plugins/Experimental/Compositing/CompositeCore/Source/CompositeCore/Public/Passes/CompositeCorePassProxy.h
@@ -80,6 +80,9 @@ namespace UE
 
 		/** Pass input declaration, referring to either internally or externally managed textures. */
 		using FPassInputDecl = TVariant<FPassInternalResourceDesc, FPassExternalResourceDesc>;
+
+		/** Enum to facilitate the declaration of passes that have one default internal input. */
+		enum EDefaultPassInputDecl { DefaultPassInputDecl };
 		
 		/** Array of pass input declarations. */
 		using FPassInputDeclArray = TArray<FPassInputDecl>;
@@ -147,6 +150,12 @@ namespace UE
 			/** Number of pass inputs */
 			int32 Num() const { return Inputs.Num(); }
 
+			/** Input array getter. */
+			TArray<UE::CompositeCore::FPassInput>& GetArray() { return Inputs; }
+
+			/** Input array const getter. */
+			const TArray<UE::CompositeCore::FPassInput>& GetArray() const { return Inputs; }
+
 			/** Conversion function to engine post-process (material) inputs. */
 			FPostProcessMaterialInputs ToPostProcessInputs(FRDGBuilder& GraphBuilder, FSceneTextureShaderParameters SceneTextures) const;
 
@@ -243,6 +252,9 @@ public:
 	/** Default constructor. */
 	FCompositeCorePassProxy() = default;
 
+	/** Constructor for a single default internal input. */
+	FCompositeCorePassProxy(UE::CompositeCore::EDefaultPassInputDecl);
+
 	/** Constructor. */
 	FCompositeCorePassProxy(UE::CompositeCore::FPassInputDeclArray InPassDeclaredInputs, TOptional<UE::CompositeCore::ResourceId> InPassOutputOverride = {});
 
