CompositeCore: Lens distortion fixes & improvements.

* Consistent behavior between the two engine lens distortion locations. TSR means output must be distorted, (secondary) upscale means output should be undistorted, since it is re-distorted later by the engine.
* Merge proxy can disable distortion handling, and simply pass-through primary input distortion.
* Removal of bOutputDistortionEnabled from the pass context struct: not needed.

Composite:
* Only main render layers use distortion, others are pass-through.
* Added media texture as optional input on plate layer (marked as already distorted).

diff --git a/Engine/Plugins/Experimental/Compositing/CompositeCore/Source/CompositeCore/Private/CompositeCoreSceneViewExtension.cpp b/Engine/Plugins/Experimental/Compositing/CompositeCore/Source/CompositeCore/Private/CompositeCoreSceneViewExtension.cpp
--- a/Engine/Plugins/Experimental/Compositing/CompositeCore/Source/CompositeCore/Private/CompositeCoreSceneViewExtension.cpp
+++ b/Engine/Plugins/Experimental/Compositing/CompositeCore/Source/CompositeCore/Private/CompositeCoreSceneViewExtension.cpp
@@ -12,7 +12,6 @@
 #include "EngineUtils.h"
 #include "Engine/Texture.h"
 #include "HDRHelper.h"
-#include "PostProcess/LensDistortion.h"
 #include "PostProcess/PostProcessAA.h"
 #include "PostProcess/PostProcessMaterialInputs.h"
 #include "Rendering/CustomRenderPass.h"
@@ -511,8 +510,6 @@ bool FCompositeCoreSceneViewExtension::ApplyPasses_Recursive(
 		{
 			// Invert alpha when writing back to scene color
 			PassContext.bOutputSceneColor = true;
-			// Optionally apply distortion
-			PassContext.bOutputDistortionEnabled = PassContext.Location > ISceneViewExtension::EPostProcessingPass::SSRInput;
 		}
 		else
 		{
@@ -541,12 +538,8 @@ FScreenPassTexture FCompositeCoreSceneViewExtension::PostProcessWork_RenderThrea
 	PassContext.OutputViewRect = Inputs.GetInput(EPostProcessMaterialInput::SceneColor).ViewRect;
 	PassContext.Location = InLocation;
 	PassContext.bOutputSceneColor = false;
-	PassContext.bOutputDistortionEnabled = false;
 
-	const FLensDistortionLUT& LensDistortionLUT = LensDistortion::GetLUTUnsafe(InView);
-	const bool bLensDistortionEnabled = LensDistortionLUT.IsEnabled() && (LensDistortion::GetPassLocationUnsafe(InView) == LensDistortion::EPassLocation::TSR);
-
-	FPassInputArray ResolvedInputs(GraphBuilder, InView, Inputs, InLocation, bLensDistortionEnabled);
+	FPassInputArray ResolvedInputs(GraphBuilder, InView, Inputs, InLocation);
 	FPassOutput Output;
 	
 	constexpr int32 RecursionLevel = 0;
diff --git a/Engine/Plugins/Experimental/Compositing/CompositeCore/Source/CompositeCore/Private/Passes/CompositeCorePassMergeProxy.cpp b/Engine/Plugins/Experimental/Compositing/CompositeCore/Source/CompositeCore/Private/Passes/CompositeCorePassMergeProxy.cpp
--- a/Engine/Plugins/Experimental/Compositing/CompositeCore/Source/CompositeCore/Private/Passes/CompositeCorePassMergeProxy.cpp
+++ b/Engine/Plugins/Experimental/Compositing/CompositeCore/Source/CompositeCore/Private/Passes/CompositeCorePassMergeProxy.cpp
@@ -71,9 +71,10 @@ namespace UE
 
 		FCompositeCoreMergeShader::FInputPassTextureParameters GetPassTextureParameters(
 			const UE::CompositeCore::FPassInput& Input,
+			const FResourceMetadata& OutputMetadata,
 			const FPassContext& PassContext,
 			const FSceneViewFamily* Family,
-			const bool bUseLensDistortion,
+			const bool bIsLensDistortionActive,
 			const float InPreExposure)
 		{
 			FCompositeCoreMergeShader::FInputPassTextureParameters Parameters;
@@ -93,23 +94,21 @@ namespace UE
 				}
 			}
 
-			EDistortionUV DistortionUV = EDistortionUV::None;
+			Parameters.DistortionUV = static_cast<uint32>(EDistortionUV::None);
 
-			if (bUseLensDistortion)
+			if (bIsLensDistortionActive)
 			{
-				const bool bInputIsDistorted = Input.Metadata.bDistorted;
-				const bool bOutputIsDistorted = PassContext.bOutputDistortionEnabled;
+				const bool bDistortedInput = Input.Metadata.bDistorted;
+				const bool bDistortedOutput = OutputMetadata.bDistorted;
 
 				// Distortion mismatch?
-				if (bInputIsDistorted != bOutputIsDistorted)
+				if (bDistortedInput != bDistortedOutput)
 				{
-					// Sample according to input distortion
-					DistortionUV = bInputIsDistorted ? EDistortionUV::Distorted : EDistortionUV::Undistorted;
+					// Sample texture according to its input distortion
+					Parameters.DistortionUV = static_cast<uint32>(bDistortedInput ? EDistortionUV::Distorted : EDistortionUV::Undistorted);
 				}
 			}
 
-			Parameters.DistortionUV = static_cast<uint32>(DistortionUV);
-
 			return Parameters;
 		}
 
@@ -120,16 +119,31 @@ namespace UE
 
 			FGlobalShaderMap* GlobalShaderMap = GetGlobalShaderMap(InView.GetFeatureLevel());
 			const FSceneViewFamily* Family = InView.Family;
+			FResourceMetadata OutMetadata = {};
 
 			const FLensDistortionLUT& LensDistortionLUT = LensDistortion::GetLUTUnsafe(InView);
-			const bool bUseLensDistortion = LensDistortionLUT.IsEnabled() && (LensDistortion::GetPassLocationUnsafe(InView) == LensDistortion::EPassLocation::TSR);
+			const bool bIsLensDistortionActive = bUseLensDistortion && LensDistortionLUT.IsEnabled()
+				&& PassContext.Location >= ISceneViewExtension::EPostProcessingPass::SSRInput;
+
+			if (bIsLensDistortionActive && PassContext.bOutputSceneColor)
+			{
+				const LensDistortion::EPassLocation EngineLocation = LensDistortion::GetPassLocationUnsafe(InView);
+
+				// Optionally set the output as distorted on passes after SSR input when TSR distorts.
+				OutMetadata.bDistorted = (PassContext.Location > ISceneViewExtension::EPostProcessingPass::SSRInput) &&
+					(EngineLocation == LensDistortion::EPassLocation::TSR);
+			}
+			else
+			{
+				// Distortion passthrough
+				OutMetadata.bDistorted = Inputs[0].Metadata.bDistorted;
+			}
 
-			FResourceMetadata OutMetadata = {};
 
 			if (PassContext.bOutputSceneColor)
 			{
 				OutMetadata.bInvertedAlpha = true;
-				OutMetadata.bDistorted = PassContext.bOutputDistortionEnabled && bUseLensDistortion;
+
 				if (PassContext.Location >= ISceneViewExtension::EPostProcessingPass::Tonemap)
 				{
 					if ((Family->EngineShowFlags.Tonemapper == 0) || (Family->EngineShowFlags.PostProcessing == 0))
@@ -157,9 +171,9 @@ namespace UE
 			PassParameters->Input0 = GetScreenPassTextureViewportParameters(FScreenPassTextureViewport(Inputs[0].Texture));
 			PassParameters->Input1 = GetScreenPassTextureViewportParameters(FScreenPassTextureViewport(Inputs[1].Texture));
 			PassParameters->Output = GetScreenPassTextureViewportParameters(FScreenPassTextureViewport(Output));
-			PassParameters->Tex0 = GetPassTextureParameters(Inputs[0], PassContext, Family, bUseLensDistortion, PreExposure);
-			PassParameters->Tex1 = GetPassTextureParameters(Inputs[1], PassContext, Family, bUseLensDistortion, PreExposure);
-			if (bUseLensDistortion)
+			PassParameters->Tex0 = GetPassTextureParameters(Inputs[0], OutMetadata, PassContext, Family, bIsLensDistortionActive, PreExposure);
+			PassParameters->Tex1 = GetPassTextureParameters(Inputs[1], OutMetadata, PassContext, Family, bIsLensDistortionActive, PreExposure);
+			if (bIsLensDistortionActive)
 			{
 				PassParameters->DistortingDisplacementTexture = LensDistortionLUT.DistortingDisplacementTexture;
 				PassParameters->UndistortingDisplacementTexture = LensDistortionLUT.UndistortingDisplacementTexture;
@@ -178,7 +192,7 @@ namespace UE
 			PassParameters->RenderTargets[0] = Output.GetRenderTargetBinding();
 
 			FCompositeCoreMergeShader::FPermutationDomain PermutationVector;
-			PermutationVector.Set<FCompositeCoreMergeShader::FUseDistortion>(bUseLensDistortion);
+			PermutationVector.Set<FCompositeCoreMergeShader::FUseDistortion>(bIsLensDistortionActive);
 			TShaderMapRef<FCompositeCoreMergeShader> PixelShader(GlobalShaderMap, PermutationVector);
 
 			FPixelShaderUtils::AddFullscreenPass(
diff --git a/Engine/Plugins/Experimental/Compositing/CompositeCore/Source/CompositeCore/Private/Passes/CompositeCorePassProxy.cpp b/Engine/Plugins/Experimental/Compositing/CompositeCore/Source/CompositeCore/Private/Passes/CompositeCorePassProxy.cpp
--- a/Engine/Plugins/Experimental/Compositing/CompositeCore/Source/CompositeCore/Private/Passes/CompositeCorePassProxy.cpp
+++ b/Engine/Plugins/Experimental/Compositing/CompositeCore/Source/CompositeCore/Private/Passes/CompositeCorePassProxy.cpp
@@ -7,6 +7,7 @@
 #include "FXRenderingUtils.h"
 #include "RenderGraphBuilder.h"
 #include "RenderGraphUtils.h"
+#include "PostProcess/LensDistortion.h"
 #include "PostProcess/PostProcessMaterialInputs.h"
 
 namespace UE
@@ -98,8 +99,7 @@ UE::CompositeCore::FPassInputArray::FPassInputArray(
 	FRDGBuilder& GraphBuilder,
 	const FSceneView& InView,
 	const FPostProcessMaterialInputs& InPostInputs,
-	const ISceneViewExtension::EPostProcessingPass& InLocation,
-	const bool bLensDistortionEnabled)
+	const ISceneViewExtension::EPostProcessingPass& InLocation)
 {
 	using namespace UE::CompositeCore;
 
@@ -110,7 +110,12 @@ UE::CompositeCore::FPassInputArray::FPassInputArray(
 		Input.Texture = FScreenPassTexture::CopyFromSlice(GraphBuilder, InPostInputs.GetInput(EPostProcessMaterialInput::SceneColor));
 		Input.Metadata.bInvertedAlpha = true;
 		Input.Metadata.bPreExposed = true;
-		Input.Metadata.bDistorted = bLensDistortionEnabled && InLocation >= ISceneViewExtension::EPostProcessingPass::SSRInput;
+		
+		// Note: This assumes the lens file is using the "SVE" method, the PPM one isn't engine-registered.
+		const FLensDistortionLUT& LensDistortionLUT = LensDistortion::GetLUTUnsafe(InView);
+		Input.Metadata.bDistorted = LensDistortionLUT.IsEnabled()
+			&& (LensDistortion::GetPassLocationUnsafe(InView) == LensDistortion::EPassLocation::TSR)
+			&& InLocation >= ISceneViewExtension::EPostProcessingPass::SSRInput;
 
 		// After-tonemap scene color may have encoding manually applied, as opposed to _SRGB textures
 		if (InLocation >= ISceneViewExtension::EPostProcessingPass::Tonemap)
diff --git a/Engine/Plugins/Experimental/Compositing/CompositeCore/Source/CompositeCore/Public/Passes/CompositeCorePassMergeProxy.h b/Engine/Plugins/Experimental/Compositing/CompositeCore/Source/CompositeCore/Public/Passes/CompositeCorePassMergeProxy.h
--- a/Engine/Plugins/Experimental/Compositing/CompositeCore/Source/CompositeCore/Public/Passes/CompositeCorePassMergeProxy.h
+++ b/Engine/Plugins/Experimental/Compositing/CompositeCore/Source/CompositeCore/Public/Passes/CompositeCorePassMergeProxy.h
@@ -30,6 +30,9 @@ namespace UE
 
 			/** Merge operation. */
 			ECompositeCoreMergeOp MergeOp = ECompositeCoreMergeOp::Over;
+
+			/** Flag to enable or disable application of engine-provided lens distortion LUTs, true by default. */
+			bool bUseLensDistortion = true;
 		};
 	}
 }
diff --git a/Engine/Plugins/Experimental/Compositing/CompositeCore/Source/CompositeCore/Public/Passes/CompositeCorePassProxy.h b/Engine/Plugins/Experimental/Compositing/CompositeCore/Source/CompositeCore/Public/Passes/CompositeCorePassProxy.h
--- a/Engine/Plugins/Experimental/Compositing/CompositeCore/Source/CompositeCore/Public/Passes/CompositeCorePassProxy.h
+++ b/Engine/Plugins/Experimental/Compositing/CompositeCore/Source/CompositeCore/Public/Passes/CompositeCorePassProxy.h
@@ -125,8 +125,7 @@ namespace UE
 				FRDGBuilder& GraphBuilder,
 				const FSceneView& InView,
 				const FPostProcessMaterialInputs& InPostInputs,
-				const ISceneViewExtension::EPostProcessingPass& InLocation,
-				const bool bInIsDistorted
+				const ISceneViewExtension::EPostProcessingPass& InLocation
 			);
 
 			/** Array access operator overload */
@@ -200,9 +199,6 @@ namespace UE
 			
 			/** Is the current pass expected to output scene color? */
 			bool bOutputSceneColor = false;
-			
-			/** Is the current pass expected to output (lens) distorted data? */
-			bool bOutputDistortionEnabled = false;
 		};
 
 		/** Options to control the built-in custom render pass. */
