CompositeCore: Shader renaming for module consistency.

diff --git a/Engine/Plugins/Experimental/Compositing/CompositeCore/Shaders/Private/CompositeCoreDisplayTransform.usf b/Engine/Plugins/Experimental/Compositing/CompositeCore/Shaders/Private/CompositeCoreDisplayTransform.usf
new file mode 100644
--- /dev/null
+++ b/Engine/Plugins/Experimental/Compositing/CompositeCore/Shaders/Private/CompositeCoreDisplayTransform.usf
@@ -0,0 +1,59 @@
+// Copyright Epic Games, Inc. All Rights Reserved.
+
+#include "/Engine/Private/Common.ush"
+
+Texture2D InputTexture;
+SamplerState InputSampler;
+uint bIsForward;
+float Gamma;
+float InvGamma;
+
+#define VERSION 0
+#define R_K 1.0
+
+// Simple Reinhard tone curve
+float3 Reinhard(float3 ColorHDR, float K = R_K)
+{
+    return ColorHDR / (ColorHDR + K);
+}
+
+// Inverse simple Reinhard tone curve
+float3 InvReinhard(float3 ColorSDR, float K = R_K)
+{
+    return K * ColorSDR / saturate(1.0 - ColorSDR);
+}
+
+// Gamma encode (faster approximation of sRGB)
+float3 GammaEncode(float3 Color)
+{
+    return pow(Color, InvGamma);
+}
+
+// Gamma decode transfer function (faster approximation of sRGB)
+float3 GammaDecode(float3 Color)
+{
+    return pow(Color, Gamma);
+}
+
+void MainPS(
+    noperspective float4 InUVAndScreenPos : TEXCOORD0,
+    out float4 OutColor : SV_Target0
+    )
+{
+    OutColor = InputTexture.Sample(InputSampler, InUVAndScreenPos.xy);
+
+#if VERSION == 0
+    if(bIsForward)
+    {
+        OutColor.rgb = GammaEncode(Reinhard(OutColor.rgb));
+    }
+    else
+    {
+        OutColor.rgb = InvReinhard(GammaDecode(OutColor.rgb));
+    }
+#elif VERSION == 1
+    OutColor.rgb = bIsForward ? Reinhard(OutColor.rgb) : InvReinhard(OutColor.rgb);
+#elif VERSION == 2
+    OutColor.rgb = bIsForward ? GammaEncode(OutColor.rgb) : GammaDecode(OutColor.rgb);
+#endif
+}
diff --git a/Engine/Plugins/Experimental/Compositing/CompositeCore/Shaders/Private/CompositeDisplayTransform.usf b/Engine/Plugins/Experimental/Compositing/CompositeCore/Shaders/Private/CompositeDisplayTransform.usf
deleted file mode 100644
diff --git a/Engine/Plugins/Experimental/Compositing/CompositeCore/Source/CompositeCore/Private/Passes/CompositeCorePassFXAAProxy.cpp b/Engine/Plugins/Experimental/Compositing/CompositeCore/Source/CompositeCore/Private/Passes/CompositeCorePassFXAAProxy.cpp
--- a/Engine/Plugins/Experimental/Compositing/CompositeCore/Source/CompositeCore/Private/Passes/CompositeCorePassFXAAProxy.cpp
+++ b/Engine/Plugins/Experimental/Compositing/CompositeCore/Source/CompositeCore/Private/Passes/CompositeCorePassFXAAProxy.cpp
@@ -26,7 +26,7 @@ class FCompositeCoreDisplayTransformShader : public FGlobalShader
 		RENDER_TARGET_BINDING_SLOTS()
 	END_SHADER_PARAMETER_STRUCT()
 };
-IMPLEMENT_GLOBAL_SHADER(FCompositeCoreDisplayTransformShader, "/Plugin/CompositeCore/Private/CompositeDisplayTransform.usf", "MainPS", SF_Pixel);
+IMPLEMENT_GLOBAL_SHADER(FCompositeCoreDisplayTransformShader, "/Plugin/CompositeCore/Private/CompositeCoreDisplayTransform.usf", "MainPS", SF_Pixel);
 
 namespace UE
 {
