CompositeCore: Add CompositeCore.Debug.DilationSize for control over pixel dilation size.

Edge case fix for motion graphics use case with fully translucent background & FXAA.

diff --git a/Engine/Plugins/Experimental/Compositing/CompositeCore/Shaders/Private/CompositeCoreDilate.usf b/Engine/Plugins/Experimental/Compositing/CompositeCore/Shaders/Private/CompositeCoreDilate.usf
--- a/Engine/Plugins/Experimental/Compositing/CompositeCore/Shaders/Private/CompositeCoreDilate.usf
+++ b/Engine/Plugins/Experimental/Compositing/CompositeCore/Shaders/Private/CompositeCoreDilate.usf
@@ -6,11 +6,11 @@
     #error TILE_SIZE must be defined
 #endif
 
-#define KERNEL_COUNT 8
+#define KERNEL_COUNT ((1 + 2 * DILATION_SIZE) * (1 + 2 * DILATION_SIZE) - 1)
 #define NEARLY_TRANSLUCENT_ALPHA 0.9999999f // Alpha is inverted opacity
 
 // Number of texels arround the group tile for scene color.
-#define LDS_TILE_BORDER_SIZE (1)
+#define LDS_TILE_BORDER_SIZE DILATION_SIZE
 
 // Width in texels of the depth tile cached into LDS.
 #define LDS_TILE_WIDTH (TILE_SIZE + 2 * LDS_TILE_BORDER_SIZE)
@@ -25,6 +25,7 @@ RWTexture2D<float4> RWOutputTexture;
 uint2 Dimensions;
 uint bOpacifyOutput;
 
+#if DILATION_SIZE == 1
 static const int2 kOffsets[KERNEL_COUNT] = {
     int2(-1,  0),
     int2( 1,  0),
@@ -35,6 +36,34 @@ static const int2 kOffsets[KERNEL_COUNT] = {
     int2( 1, -1),
     int2(-1, -1),
 };
+#elif DILATION_SIZE == 2
+static const int2 kOffsets[24] = {
+    int2(-2,  0),
+    int2(-1,  0),
+    int2( 1,  0),
+    int2( 2,  0),
+    int2( 0, -2),
+    int2( 0, -1),
+    int2( 0,  1),
+    int2( 0,  2),
+    int2(-2, -1),
+    int2(-1, -2),
+    int2( 1, -2),
+    int2( 2, -1),
+    int2( 2,  1),
+    int2( 1,  2),
+    int2(-1,  2),
+    int2(-2,  1),
+    int2(-1, -1),
+    int2( 1, -1),
+    int2( 1,  1),
+    int2(-1,  1),
+    int2(-2, -2),
+    int2( 2, -2),
+    int2( 2,  2),
+    int2(-2,  2),
+};
+#endif
 
 groupshared float4 SharedColor[LDS_ARRAY_SIZE];
 
@@ -110,6 +139,7 @@ void MainCS(
 
     float4 TexelColor = SampleCachedColors(GroupThreadId);
     
+#if DILATION_SIZE > 0
 	// Apply (gather) dilation on any translucent or nearly translucent texel
     BRANCH
     if(TexelColor.a > NEARLY_TRANSLUCENT_ALPHA)
@@ -125,6 +155,7 @@ void MainCS(
             }
         }
     }
+#endif
 
 	// Condtionally & fully opacify non-translucent colors for the over blend.
 	if (bOpacifyOutput && TexelColor.a < NEARLY_TRANSLUCENT_ALPHA)
diff --git a/Engine/Plugins/Experimental/Compositing/CompositeCore/Source/CompositeCore/Private/CompositeCoreSceneViewExtension.cpp b/Engine/Plugins/Experimental/Compositing/CompositeCore/Source/CompositeCore/Private/CompositeCoreSceneViewExtension.cpp
--- a/Engine/Plugins/Experimental/Compositing/CompositeCore/Source/CompositeCore/Private/CompositeCoreSceneViewExtension.cpp
+++ b/Engine/Plugins/Experimental/Compositing/CompositeCore/Source/CompositeCore/Private/CompositeCoreSceneViewExtension.cpp
@@ -55,10 +55,7 @@ public:
 		const FRDGTextureDesc TextureDesc = FRDGTextureDesc::Create2D(RenderTargetTexture->Desc.Extent, PF_FloatRGBA, FClearValueBinding::Black, TexCreate_UAV | TexCreate_ShaderResource | TexCreate_RenderTargetable);
 		FRDGTextureRef Output = GraphBuilder.CreateTexture(TextureDesc, TEXT("CompositeCoreProcessedTexture"));
 
-		if (Inputs.DilationSize > 0)
-		{
-			UE::CompositeCore::Private::AddDilatePass(GraphBuilder, RenderTargetTexture, Output, ViewFeatureLevel, Inputs);
-		}
+		UE::CompositeCore::Private::AddDilatePass(GraphBuilder, RenderTargetTexture, Output, ViewFeatureLevel, Inputs);
 
 		ParentExtension->CollectCustomRenderTarget(ViewId, GraphBuilder.ConvertToExternalTexture(Output));
 	}
@@ -74,7 +71,6 @@ private:
 FCompositeCoreSceneViewExtension::FCompositeCoreSceneViewExtension(const FAutoRegister& AutoReg, UWorld* InWorld)
 	: FWorldSceneViewExtension(AutoReg, InWorld)
 {
-	BuiltInRenderPassOptions = UE::CompositeCore::FBuiltInRenderPassOptions{};
 }
 
 FCompositeCoreSceneViewExtension::~FCompositeCoreSceneViewExtension() = default;
@@ -173,6 +169,11 @@ void FCompositeCoreSceneViewExtension::SetBuiltInRenderPassOptions_GameThread(co
 	BuiltInRenderPassOptions = InOptions;
 }
 
+void FCompositeCoreSceneViewExtension::ResetBuiltInRenderPassOptions_GameThread()
+{
+	BuiltInRenderPassOptions.Reset();
+}
+
 /* Called by the custom render pass to store its view render target for this frame. */
 
 void FCompositeCoreSceneViewExtension::CollectCustomRenderTarget(uint32 InViewId, const TRefCountPtr<IPooledRenderTarget>& InRenderTarget)
@@ -269,24 +270,26 @@ void FCompositeCoreSceneViewExtension::BeginRenderViewFamily(FSceneViewFamily& I
 			return;
 		}
 
+		UE::CompositeCore::FBuiltInRenderPassOptions RenderPassOptions = BuiltInRenderPassOptions.IsSet() ? *BuiltInRenderPassOptions : UE::CompositeCore::FBuiltInRenderPassOptions{};
+
 		// Create a new custom render pass to render the composite primitive(s)
 		FCompositeCoreCustomRenderPass* CustomRenderPass = new FCompositeCoreCustomRenderPass(
 			InView.UnscaledViewRect.Size(),
 			this,
 			InView,
-			BuiltInRenderPassOptions
+			RenderPassOptions
 		);
 
 		FSceneInterface::FCustomRenderPassRendererInput PassInput{};
 		PassInput.EngineShowFlags = InViewFamily.EngineShowFlags;
 		PassInput.EngineShowFlags.DisableFeaturesForUnlit();
 		PassInput.EngineShowFlags.SetTranslucency(true);
-		PassInput.EngineShowFlags.SetUnlitViewmode(BuiltInRenderPassOptions.bEnableUnlitViewmode);
+		PassInput.EngineShowFlags.SetUnlitViewmode(RenderPassOptions.bEnableUnlitViewmode);
 		PassInput.EngineShowFlags.SetAllowPrimitiveAlphaHoldout(false);
-		if (BuiltInRenderPassOptions.ViewUserFlagsOverride.IsSet())
+		if (RenderPassOptions.ViewUserFlagsOverride.IsSet())
 		{
 			PassInput.bOverridesPostVolumeUserFlags = true;
-			PassInput.PostVolumeUserFlags = BuiltInRenderPassOptions.ViewUserFlagsOverride.GetValue();
+			PassInput.PostVolumeUserFlags = RenderPassOptions.ViewUserFlagsOverride.GetValue();
 		}
 		// Note: Incoming view location is invalid for scene captures
 		PassInput.ViewLocation = InView.ViewMatrices.GetViewOrigin();
diff --git a/Engine/Plugins/Experimental/Compositing/CompositeCore/Source/CompositeCore/Private/CompositeCoreSceneViewExtension.h b/Engine/Plugins/Experimental/Compositing/CompositeCore/Source/CompositeCore/Private/CompositeCoreSceneViewExtension.h
--- a/Engine/Plugins/Experimental/Compositing/CompositeCore/Source/CompositeCore/Private/CompositeCoreSceneViewExtension.h
+++ b/Engine/Plugins/Experimental/Compositing/CompositeCore/Source/CompositeCore/Private/CompositeCoreSceneViewExtension.h
@@ -51,8 +51,11 @@ public:
 	/** Reset post-processing render work. */
 	void ResetPostRenderWork_GameThread();
 
-	/** Set an optional view user flags override. */
+	/** Set options for built-in composite custom render pass. */
 	void SetBuiltInRenderPassOptions_GameThread(const UE::CompositeCore::FBuiltInRenderPassOptions& InOptions);
+	
+	/** Reset default options for built-in composite custom render pass. */
+	void ResetBuiltInRenderPassOptions_GameThread();
 
 	/** Called by the custom render pass to store its view render target for this frame. */
 	void CollectCustomRenderTarget(uint32 InViewId, const TRefCountPtr<IPooledRenderTarget>& InRenderTarget);
@@ -116,9 +119,6 @@ private:
 	/** Collection of primitives to render as a custom render pass and composite after post-processing. .*/
 	TSet<TWeakObjectPtr<UPrimitiveComponent>> CompositePrimitives;
 
-	/** Built-in composite render pass options. .*/
-	UE::CompositeCore::FBuiltInRenderPassOptions BuiltInRenderPassOptions;
-
 	/** Array of external texture inputs. */
 	TMap<UE::CompositeCore::ResourceId, UE::CompositeCore::FExternalRenderTarget> ExternalInputs_RenderThread;
 
@@ -127,5 +127,8 @@ private:
 
 	/** Custom render work. */
 	TOptional<UE::CompositeCore::FPostRenderWork> PostRenderWork_RenderThread;
+
+	/** Built-in composite render pass options. .*/
+	TOptional<UE::CompositeCore::FBuiltInRenderPassOptions> BuiltInRenderPassOptions;
 };
 
diff --git a/Engine/Plugins/Experimental/Compositing/CompositeCore/Source/CompositeCore/Private/CompositeCoreSubsystem.cpp b/Engine/Plugins/Experimental/Compositing/CompositeCore/Source/CompositeCore/Private/CompositeCoreSubsystem.cpp
--- a/Engine/Plugins/Experimental/Compositing/CompositeCore/Source/CompositeCore/Private/CompositeCoreSubsystem.cpp
+++ b/Engine/Plugins/Experimental/Compositing/CompositeCore/Source/CompositeCore/Private/CompositeCoreSubsystem.cpp
@@ -132,7 +132,7 @@ void UCompositeCoreSubsystem::ResetBuiltInRenderPassOptions()
 {
 	if (CompositeCoreViewExtension.IsValid())
 	{
-		CompositeCoreViewExtension->SetBuiltInRenderPassOptions_GameThread(UE::CompositeCore::FBuiltInRenderPassOptions{});
+		CompositeCoreViewExtension->ResetBuiltInRenderPassOptions_GameThread();
 	}
 }
 
diff --git a/Engine/Plugins/Experimental/Compositing/CompositeCore/Source/CompositeCore/Private/Passes/CompositeCorePassDilate.cpp b/Engine/Plugins/Experimental/Compositing/CompositeCore/Source/CompositeCore/Private/Passes/CompositeCorePassDilate.cpp
--- a/Engine/Plugins/Experimental/Compositing/CompositeCore/Source/CompositeCore/Private/Passes/CompositeCorePassDilate.cpp
+++ b/Engine/Plugins/Experimental/Compositing/CompositeCore/Source/CompositeCore/Private/Passes/CompositeCorePassDilate.cpp
@@ -14,6 +14,9 @@ class FCompositeCoreDilateShader : public FGlobalShader
 	DECLARE_GLOBAL_SHADER(FCompositeCoreDilateShader);
 	SHADER_USE_PARAMETER_STRUCT(FCompositeCoreDilateShader, FGlobalShader);
 
+	class FDilationSize : SHADER_PERMUTATION_INT("DILATION_SIZE", 3);
+	using FPermutationDomain = TShaderPermutationDomain<FDilationSize>;
+
 	BEGIN_SHADER_PARAMETER_STRUCT(FParameters, )
 		SHADER_PARAMETER_RDG_TEXTURE(Texture2D, InputTexture)
 		SHADER_PARAMETER_RDG_TEXTURE_UAV(RWTexture2D<float4>, RWOutputTexture)
@@ -52,7 +55,10 @@ namespace UE
 				PassParameters->Dimensions = TextureSize;
 				PassParameters->bOpacifyOutput = PassInputs.bOpacifyOutput;
 
-				TShaderMapRef<FCompositeCoreDilateShader> ComputeShader(GlobalShaderMap);
+				FCompositeCoreDilateShader::FPermutationDomain PermutationVector;
+				PermutationVector.Set<FCompositeCoreDilateShader::FDilationSize>(FMath::Clamp(PassInputs.DilationSize, 0, 2));
+
+				TShaderMapRef<FCompositeCoreDilateShader> ComputeShader(GlobalShaderMap, PermutationVector);
 				FComputeShaderUtils::AddPass(
 					GraphBuilder,
 					RDG_EVENT_NAME("CompositeCore.Dilate (%dx%d)", TextureSize.X, TextureSize.Y),
diff --git a/Engine/Plugins/Experimental/Compositing/CompositeCore/Source/CompositeCore/Private/Passes/CompositeCorePassProxy.cpp b/Engine/Plugins/Experimental/Compositing/CompositeCore/Source/CompositeCore/Private/Passes/CompositeCorePassProxy.cpp
--- a/Engine/Plugins/Experimental/Compositing/CompositeCore/Source/CompositeCore/Private/Passes/CompositeCorePassProxy.cpp
+++ b/Engine/Plugins/Experimental/Compositing/CompositeCore/Source/CompositeCore/Private/Passes/CompositeCorePassProxy.cpp
@@ -10,6 +10,12 @@
 #include "PostProcess/LensDistortion.h"
 #include "PostProcess/PostProcessMaterialInputs.h"
 
+static TAutoConsoleVariable<int32> CVarCompositeCoreDebugDilationSize(
+	TEXT("CompositeCore.Debug.DilationSize"),
+	1,
+	TEXT("Size of the pixel dilation applied onto the composite custom render pass. 0, 1 & 2 are supported."),
+	ECVF_Default);
+
 namespace UE
 {
 	namespace CompositeCore
@@ -23,7 +29,7 @@ namespace UE
 		FBuiltInRenderPassOptions::FBuiltInRenderPassOptions()
 			: ViewUserFlagsOverride{}
 			, bEnableUnlitViewmode{true}
-			, DilationSize{1}
+			, DilationSize{ CVarCompositeCoreDebugDilationSize.GetValueOnGameThread() }
 			, bOpacifyOutput{true}
 		{
 		}
