CompositeCore: Fix regression where the basic holdout composite pipeline renders black in scene captures.

diff --git a/Engine/Plugins/Experimental/Compositing/CompositeCore/Source/CompositeCore/Private/CompositeCoreSceneViewExtension.cpp b/Engine/Plugins/Experimental/Compositing/CompositeCore/Source/CompositeCore/Private/CompositeCoreSceneViewExtension.cpp
--- a/Engine/Plugins/Experimental/Compositing/CompositeCore/Source/CompositeCore/Private/CompositeCoreSceneViewExtension.cpp
+++ b/Engine/Plugins/Experimental/Compositing/CompositeCore/Source/CompositeCore/Private/CompositeCoreSceneViewExtension.cpp
@@ -354,8 +354,7 @@ void FCompositeCoreSceneViewExtension::SubscribeToPostProcessingPass(EPostProces
 
 bool FCompositeCoreSceneViewExtension::IsActiveForView(const FSceneView& InView) const
 {
-	bool bIsActive = InView.Family->bIsMainViewFamily;
-	bIsActive &= !InView.bIsSceneCapture;
+	bool bIsActive = true;
 	bIsActive &= static_cast<bool>(InView.Family->EngineShowFlags.AllowPrimitiveAlphaHoldout);
 	bIsActive &= CustomRenderTargetPerView_RenderThread.Contains(InView.GetViewKey()) || !ExternalInputs_RenderThread.IsEmpty();
 
