CompositeCore: Minor fix to keep the SVE active if there is custom render work (without any holdout primitives).

This can occur if a 2D layer is composited over the main render in our Composite plugin.

diff --git a/Engine/Plugins/Experimental/Compositing/CompositeCore/Source/CompositeCore/Private/CompositeCoreSceneViewExtension.cpp b/Engine/Plugins/Experimental/Compositing/CompositeCore/Source/CompositeCore/Private/CompositeCoreSceneViewExtension.cpp
--- a/Engine/Plugins/Experimental/Compositing/CompositeCore/Source/CompositeCore/Private/CompositeCoreSceneViewExtension.cpp
+++ b/Engine/Plugins/Experimental/Compositing/CompositeCore/Source/CompositeCore/Private/CompositeCoreSceneViewExtension.cpp
@@ -119,6 +119,8 @@ void FCompositeCoreSceneViewExtension::UnregisterPrimitives_GameThread(const TAr
 
 void FCompositeCoreSceneViewExtension::SetPostRenderWork_GameThread(UE::CompositeCore::FPostRenderWork&& InWork)
 {
+	bHasCustomRenderWork = !InWork.FramePasses.IsEmpty();
+
 	ENQUEUE_RENDER_COMMAND(CopyCompositeCoreRenderWork)(
 		[InWork = MoveTemp(InWork), WeakThis = AsWeak()](FRHICommandListImmediate& RHICmdList) mutable
 		{
@@ -152,6 +154,8 @@ void FCompositeCoreSceneViewExtension::SetPostRenderWork_GameThread(UE::Composit
 
 void FCompositeCoreSceneViewExtension::ResetPostRenderWork_GameThread()
 {
+	bHasCustomRenderWork = false;
+
 	ENQUEUE_RENDER_COMMAND(CopyCompositeCoreRenderWork)(
 		[WeakThis = AsWeak()](FRHICommandListImmediate& RHICmdList)
 		{
@@ -184,7 +188,7 @@ void FCompositeCoreSceneViewExtension::CollectCustomRenderTarget(uint32 InViewId
 bool FCompositeCoreSceneViewExtension::IsActiveThisFrame_Internal(const FSceneViewExtensionContext& Context) const
 {
 	bool bIsActive = FWorldSceneViewExtension::IsActiveThisFrame_Internal(Context);
-	bIsActive &= !CompositePrimitives.IsEmpty();
+	bIsActive &= bHasCustomRenderWork || !CompositePrimitives.IsEmpty();
 	bIsActive &= !IsHDREnabled();
 
 	return bIsActive;
diff --git a/Engine/Plugins/Experimental/Compositing/CompositeCore/Source/CompositeCore/Private/CompositeCoreSceneViewExtension.h b/Engine/Plugins/Experimental/Compositing/CompositeCore/Source/CompositeCore/Private/CompositeCoreSceneViewExtension.h
--- a/Engine/Plugins/Experimental/Compositing/CompositeCore/Source/CompositeCore/Private/CompositeCoreSceneViewExtension.h
+++ b/Engine/Plugins/Experimental/Compositing/CompositeCore/Source/CompositeCore/Private/CompositeCoreSceneViewExtension.h
@@ -128,7 +128,10 @@ private:
 	/** Custom render work. */
 	TOptional<UE::CompositeCore::FPostRenderWork> PostRenderWork_RenderThread;
 
-	/** Built-in composite render pass options. .*/
+	/** Built-in composite render pass options.*/
 	TOptional<UE::CompositeCore::FBuiltInRenderPassOptions> BuiltInRenderPassOptions;
+
+	/** Used to keep the SVE active if any custom render work is present.*/
+	bool bHasCustomRenderWork = false;
 };
 
