CompositeCore: Fix incorrect setting validation for non-editor builds.

Issue where console variable-backed renderer setting members are imported in editor builds only (see ImportConsoleVariableValues()). Fix is to instead rely direcly on the cvars, which reflect the correct values.

diff --git a/Engine/Plugins/Experimental/Compositing/CompositeCore/Source/CompositeCore/Private/CompositeCoreSubsystem.cpp b/Engine/Plugins/Experimental/Compositing/CompositeCore/Source/CompositeCore/Private/CompositeCoreSubsystem.cpp
--- a/Engine/Plugins/Experimental/Compositing/CompositeCore/Source/CompositeCore/Private/CompositeCoreSubsystem.cpp
+++ b/Engine/Plugins/Experimental/Compositing/CompositeCore/Source/CompositeCore/Private/CompositeCoreSubsystem.cpp
@@ -6,6 +6,7 @@
 #include "CompositeCoreSceneViewExtension.h"
 
 #include "Engine/RendererSettings.h"
+#include "HAL/IConsoleManager.h"
 
 #if WITH_EDITOR
 #include "Framework/Notifications/NotificationManager.h"
@@ -138,14 +139,13 @@ void UCompositeCoreSubsystem::ResetBuiltInRenderPassOptions()
 
 bool UCompositeCoreSubsystem::IsProjectSettingsValid()
 {
-	const URendererSettings* RendererSettings = GetDefault<URendererSettings>();
-	check(RendererSettings);
-
-	const bool bValidSettings = 
-		RendererSettings->bEnableAlphaChannelInPostProcessing
-		&& RendererSettings->bDeferredSupportPrimitiveAlphaHoldout;
+	static const auto CVarAlphaOutput = IConsoleManager::Get().FindConsoleVariable(TEXT("r.PostProcessing.PropagateAlpha"));
+	static const auto CVarSupportPrimitiveAlphaHoldout = IConsoleManager::Get().FindConsoleVariable(TEXT("r.Deferred.SupportPrimitiveAlphaHoldout"));
+	
+	const bool bAlphaOutput = CVarAlphaOutput ? CVarAlphaOutput->GetBool() : false;
+	const bool bSupportPrimitiveAlphaHoldout = CVarSupportPrimitiveAlphaHoldout ? CVarSupportPrimitiveAlphaHoldout->GetBool() : false;
 
-	return bValidSettings;
+	return bAlphaOutput && bSupportPrimitiveAlphaHoldout;
 }
 
 bool UCompositeCoreSubsystem::ValidateProjectSettings()
