CompositeCore: Fixes to holdout composite component for motion graphics.

* Added list of allowed scene component classes for which users will not be warned if associated primitive cannot immediately be found.
* Included "Text3DComponent" as an allowed component by default since its primitive is created after our registration.
* Added "Don't show this again" option on the notification warning if a users fails to attach the holdout composite component to a primitive.
* Fixed component enable toggle which regressed in a recent CL.
* Delayed settings caching to allow all module loads to finish before attempting to load objects.

diff --git a/Engine/Plugins/Experimental/Compositing/CompositeCore/Source/CompositeCore/Private/CompositeCoreSettings.cpp b/Engine/Plugins/Experimental/Compositing/CompositeCore/Source/CompositeCore/Private/CompositeCoreSettings.cpp
--- a/Engine/Plugins/Experimental/Compositing/CompositeCore/Source/CompositeCore/Private/CompositeCoreSettings.cpp
+++ b/Engine/Plugins/Experimental/Compositing/CompositeCore/Source/CompositeCore/Private/CompositeCoreSettings.cpp
@@ -4,6 +4,7 @@
 #include "CompositeCoreSubsystem.h"
 
 #include "Components/PrimitiveComponent.h"
+#include "Misc/ScopeLock.h"
 
 #include UE_INLINE_GENERATED_CPP_BY_NAME(CompositeCoreSettings)
 
@@ -35,7 +36,16 @@ UCompositeCorePluginSettings::UCompositeCorePluginSettings()
 			FSoftClassPath("/Script/Engine.LineBatchComponent")
 		}
 	)
+	, AllowedComponentClasses(
+		{
+			FSoftClassPath("/Script/Text3D.Text3DComponent"),
+		}
+	)
 	, SceneViewExtensionPriority(100 + 1 /* OPENCOLORIO_SCENE_VIEW_EXTENSION_PRIORITY + 1 */)
+	, CachedDisabledPrimitiveClasses()
+	, bIsPrimitiveCacheDirty(true)
+	, CachedAllowedComponentClasses()
+	, bIsComponentCacheDirty(true)
 { }
 
 
@@ -61,7 +71,11 @@ void UCompositeCorePluginSettings::PostEditChangeProperty(FPropertyChangedEvent&
 
 	if (PropertyChangedEvent.GetMemberPropertyName() == GET_MEMBER_NAME_CHECKED(ThisClass, DisabledPrimitiveClasses))
 	{
-		CacheDisabledClasses();
+		bIsPrimitiveCacheDirty = true;
+	}
+	else if (PropertyChangedEvent.GetMemberPropertyName() == GET_MEMBER_NAME_CHECKED(ThisClass, AllowedComponentClasses))
+	{
+		bIsComponentCacheDirty = true;
 	}
 
 	if (PropertyChangedEvent.Property)
@@ -81,12 +95,23 @@ void UCompositeCorePluginSettings::PostInitProperties()
 		ImportConsoleVariableValues();
 	}
 #endif // #if WITH_EDITOR
-
-	CacheDisabledClasses();
 }
 
 bool UCompositeCorePluginSettings::IsAllowedPrimitiveClass(const UPrimitiveComponent* InPrimitiveComponent) const
 {
+	if (!InPrimitiveComponent)
+	{
+		return false;
+	}
+
+	// Note: Should always be called from the game thread, but we lock the mutable cache just in case.
+	FScopeLock Lock(&CriticalSection);
+
+	if (bIsPrimitiveCacheDirty)
+	{
+		CacheDisabledPrimitiveClasses();
+	}
+
 	for (const UClass* ObjectClass : CachedDisabledPrimitiveClasses)
 	{
 		if (InPrimitiveComponent->IsA(ObjectClass))
@@ -98,7 +123,43 @@ bool UCompositeCorePluginSettings::IsAllowedPrimitiveClass(const UPrimitiveCompo
 	return true;
 }
 
-void UCompositeCorePluginSettings::CacheDisabledClasses()
+bool UCompositeCorePluginSettings::IsAllowedComponentClass(const USceneComponent* InComponent) const
+{
+	if (!InComponent)
+	{
+		return false;
+	}
+
+	{
+		// Note: Should always be called from the game thread, but we lock the mutable cache just in case.
+		FScopeLock Lock(&CriticalSection);
+		
+		if (bIsComponentCacheDirty)
+		{
+			CacheAllowedComponentClasses();
+		}
+
+		for (const FSoftClassPath& ObjectClass : AllowedComponentClasses)
+		{
+			if (const UClass* AllowedClass = ObjectClass.TryLoadClass<UObject>())
+			{
+				if (InComponent->IsA(AllowedClass))
+				{
+					return true;
+				}
+			}
+		}
+	}
+
+	if (const UPrimitiveComponent* PrimitiveComponent = Cast<UPrimitiveComponent>(InComponent))
+	{
+		return IsAllowedPrimitiveClass(PrimitiveComponent);
+	}
+
+	return false;
+}
+
+void UCompositeCorePluginSettings::CacheDisabledPrimitiveClasses() const
 {
 	CachedDisabledPrimitiveClasses.Reset(0);
 
@@ -109,4 +170,21 @@ void UCompositeCorePluginSettings::CacheDisabledClasses()
 			CachedDisabledPrimitiveClasses.Add(DisabledClass);
 		}
 	}
+
+	bIsPrimitiveCacheDirty = false;
+}
+
+void UCompositeCorePluginSettings::CacheAllowedComponentClasses() const
+{
+	CachedAllowedComponentClasses.Reset(0);
+
+	for (const FSoftClassPath& ObjectClass : AllowedComponentClasses)
+	{
+		if (const UClass* AllowedClass = ObjectClass.TryLoadClass<UObject>())
+		{
+			CachedAllowedComponentClasses.Add(AllowedClass);
+		}
+	}
+
+	bIsComponentCacheDirty = false;
 }
diff --git a/Engine/Plugins/Experimental/Compositing/CompositeCore/Source/CompositeCore/Private/HoldoutCompositeComponent.cpp b/Engine/Plugins/Experimental/Compositing/CompositeCore/Source/CompositeCore/Private/HoldoutCompositeComponent.cpp
--- a/Engine/Plugins/Experimental/Compositing/CompositeCore/Source/CompositeCore/Private/HoldoutCompositeComponent.cpp
+++ b/Engine/Plugins/Experimental/Compositing/CompositeCore/Source/CompositeCore/Private/HoldoutCompositeComponent.cpp
@@ -9,6 +9,8 @@
 #include "CompositeCoreSubsystem.h"
 
 #if WITH_EDITOR
+#include "CoreGlobals.h"
+#include "Misc/ConfigCacheIni.h"
 #include "Framework/Notifications/NotificationManager.h"
 #include "Widgets/Notifications/SNotificationList.h"
 #endif
@@ -31,8 +33,9 @@ namespace UE
 
 				if (IsValid(InParentComponent))
 				{
-					if (UPrimitiveComponent* PrimitiveComponent = Cast<UPrimitiveComponent>(InParentComponent))
 					{
+						UPrimitiveComponent* PrimitiveComponent = Cast<UPrimitiveComponent>(InParentComponent);
+						
 						if (PluginSettings->IsAllowedPrimitiveClass(PrimitiveComponent))
 						{
 							OutPrimitiveComponents.Add(PrimitiveComponent);
@@ -44,12 +47,11 @@ namespace UE
 
 					for (USceneComponent* ParentChild : ParentChildComponents)
 					{
-						if (UPrimitiveComponent* PrimitiveComponent = Cast<UPrimitiveComponent>(ParentChild))
+						UPrimitiveComponent* PrimitiveComponent = Cast<UPrimitiveComponent>(ParentChild);
+						
+						if (PluginSettings->IsAllowedPrimitiveClass(PrimitiveComponent))
 						{
-							if (PluginSettings->IsAllowedPrimitiveClass(PrimitiveComponent))
-							{
-								OutPrimitiveComponents.Add(PrimitiveComponent);
-							}
+							OutPrimitiveComponents.Add(PrimitiveComponent);
 						}
 					}
 				}
@@ -120,6 +122,25 @@ void UHoldoutCompositeComponent::DetachFromComponent(const FDetachmentTransformR
 	Super::DetachFromComponent(DetachmentRules);
 }
 
+#if WITH_EDITOR
+void UHoldoutCompositeComponent::PostEditChangeProperty(FPropertyChangedEvent& PropertyChangedEvent)
+{
+	Super::PostEditChangeProperty(PropertyChangedEvent);
+
+	if (PropertyChangedEvent.GetMemberPropertyName() == GET_MEMBER_NAME_CHECKED(ThisClass, bIsEnabled))
+	{
+		if (bIsEnabled)
+		{
+			RegisterCompositeImpl();
+		}
+		else
+		{
+			UnregisterCompositeImpl();
+		}
+	}
+}
+#endif
+
 void UHoldoutCompositeComponent::OnAttachmentChanged()
 {
 	Super::OnAttachmentChanged();
@@ -133,11 +154,48 @@ void UHoldoutCompositeComponent::OnAttachmentChanged()
 		if (ParentPrimitives.IsEmpty())
 		{
 #if WITH_EDITOR
-			FNotificationInfo Info(LOCTEXT("CompositeParentNotification",
-				"The composite component must be parented to a primitive component (or one that has primitives)."));
-			Info.ExpireDuration = 5.0f;
+			const UCompositeCorePluginSettings* PluginSettings = GetDefault<UCompositeCorePluginSettings>();
+			
+			if (!PluginSettings->IsAllowedComponentClass(SceneComponent))
+			{
+				/** Utility functions for notifications */
+				struct FSuppressDialogOptions
+				{
+					static bool ShouldSuppressModal()
+					{
+						bool bSuppressNotification = false;
+						GConfig->GetBool(TEXT("CompositeCore"), TEXT("SuppressCompositeCorePrimitiveWarning"), bSuppressNotification, GEditorPerProjectIni);
+						return bSuppressNotification;
+					}
+
+					static ECheckBoxState GetDontAskAgainCheckBoxState()
+					{
+						return ShouldSuppressModal() ? ECheckBoxState::Checked : ECheckBoxState::Unchecked;
+					}
 
-			FSlateNotificationManager::Get().AddNotification(Info);
+					static void OnDontAskAgainCheckBoxStateChanged(ECheckBoxState NewState)
+					{
+						// If the user selects to not show this again, set that in the config so we know about it in between sessions
+						const bool bSuppressNotification = (NewState == ECheckBoxState::Checked);
+						GConfig->SetBool(TEXT("CompositeCore"), TEXT("SuppressCompositeCorePrimitiveWarning"), bSuppressNotification, GEditorPerProjectIni);
+					}
+				};
+
+				// Skip if the user has specified to supress this pop up.
+				if (!FSuppressDialogOptions::ShouldSuppressModal())
+				{
+					FNotificationInfo Info(LOCTEXT("CompositeParentNotification",
+						"The composite component must be parented to a primitive component (or one that has primitives)."));
+					Info.ExpireDuration = 5.0f;
+
+					// Add a "Don't show this again" option
+					Info.CheckBoxState = TAttribute<ECheckBoxState>::Create(&FSuppressDialogOptions::GetDontAskAgainCheckBoxState);
+					Info.CheckBoxStateChanged = FOnCheckStateChanged::CreateStatic(&FSuppressDialogOptions::OnDontAskAgainCheckBoxStateChanged);
+					Info.CheckBoxText = LOCTEXT("DontShowThisAgainCheckBoxMessage", "Don't show this again");
+
+					FSlateNotificationManager::Get().AddNotification(Info);
+				}
+			}
 #endif
 		}
 		else
diff --git a/Engine/Plugins/Experimental/Compositing/CompositeCore/Source/CompositeCore/Public/CompositeCoreSettings.h b/Engine/Plugins/Experimental/Compositing/CompositeCore/Source/CompositeCore/Public/CompositeCoreSettings.h
--- a/Engine/Plugins/Experimental/Compositing/CompositeCore/Source/CompositeCore/Public/CompositeCoreSettings.h
+++ b/Engine/Plugins/Experimental/Compositing/CompositeCore/Source/CompositeCore/Public/CompositeCoreSettings.h
@@ -10,6 +10,7 @@
 #include "CompositeCoreSettings.generated.h"
 
 class UPrimitiveComponent;
+class USceneComponent;
 
 namespace UE
 {
@@ -43,9 +44,12 @@ public:
 	virtual void PostInitProperties() override;
 	//~ End UObject Interface
 
-	// Ensure that the primitive class is allowed by checking against plugin settings.
+	// Ensure that the primitive component class is allowed by checking against plugin settings.
 	COMPOSITECORE_API bool IsAllowedPrimitiveClass(const UPrimitiveComponent* InPrimitiveComponent) const;
 
+	// Ensure that the scene component class is allowed by checking against plugin settings.
+	COMPOSITECORE_API bool IsAllowedComponentClass(const USceneComponent* InComponent) const;
+
 public:
 	UPROPERTY(config, EditAnywhere, Category = CompositeCore, meta = (
 		ConsoleVariable = "CompositeCore.ApplyPreExposure", DisplayName = "Apply Pre-Exposure",
@@ -63,17 +67,36 @@ public:
 	UPROPERTY(config, EditAnywhere, Category = CompositeCore)
 	TArray<FSoftClassPath> DisabledPrimitiveClasses;
 
+	/** Allowed component classes for which users will not be warned if associated primitive cannot immediately be found. */
+	UPROPERTY(config, EditAnywhere, Category = CompositeCore)
+	TArray<FSoftClassPath> AllowedComponentClasses;
+
 	/** Composite (scene view extension) post-processing priority, which defaults to before OpenColorIO. */
 	UPROPERTY(config, EditAnywhere, Category = CompositeCore, AdvancedDisplay)
 	int32 SceneViewExtensionPriority;
 
 private:
 
-	/** Cache disable class pointers to prevent load object calls during allowed class checks. */
-	void CacheDisabledClasses();
-	
-	/** Cached list of disallowed class types. */
-	TArray<const UClass*> CachedDisabledPrimitiveClasses;
+	/** Cache primitive class pointers to prevent load object calls during allowed class checks. */
+	void CacheDisabledPrimitiveClasses() const;
+
+	/** Cache component class pointers to prevent load object calls during allowed class checks. */
+	void CacheAllowedComponentClasses() const;
+
+	/** Cached list of disallowed primitive class types. */
+	mutable TArray<const UClass*> CachedDisabledPrimitiveClasses;
+
+	/** Dirty primitive cache flag. */
+	mutable bool bIsPrimitiveCacheDirty;
+
+	/** Cached list of allowed component class types. */
+	mutable TArray<const UClass*> CachedAllowedComponentClasses;
+
+	/** Dirty component cache flag. */
+	mutable bool bIsComponentCacheDirty;
+
+	/** Mutable cache lock. */
+	mutable FCriticalSection CriticalSection;
 };
 
 
diff --git a/Engine/Plugins/Experimental/Compositing/CompositeCore/Source/CompositeCore/Public/HoldoutCompositeComponent.h b/Engine/Plugins/Experimental/Compositing/CompositeCore/Source/CompositeCore/Public/HoldoutCompositeComponent.h
--- a/Engine/Plugins/Experimental/Compositing/CompositeCore/Source/CompositeCore/Public/HoldoutCompositeComponent.h
+++ b/Engine/Plugins/Experimental/Compositing/CompositeCore/Source/CompositeCore/Public/HoldoutCompositeComponent.h
@@ -28,6 +28,10 @@ public:
 	virtual void DetachFromComponent(const FDetachmentTransformRules& DetachmentRules) override;
 	//~ End USceneComponent Interface
 
+#if WITH_EDITOR
+	virtual void PostEditChangeProperty(struct FPropertyChangedEvent& PropertyChangedEvent) override;
+#endif //WITH_EDITOR
+
 	/* Get the enabled state of the component. */
 	UFUNCTION(BlueprintGetter)
 	COMPOSITECORE_API bool IsEnabled() const;
