CompositeCore: Add API support for resource pre-processing.

* Preprocessing passes can be registered to run once per frame, at the start of the very first frame render.
* Fixed output override support in the FXAA pass.
* FPostRenderWork can be renamed in a subsequent CL.

diff --git a/Engine/Plugins/Experimental/Compositing/CompositeCore/Source/CompositeCore/Private/CompositeCoreSceneViewExtension.cpp b/Engine/Plugins/Experimental/Compositing/CompositeCore/Source/CompositeCore/Private/CompositeCoreSceneViewExtension.cpp
--- a/Engine/Plugins/Experimental/Compositing/CompositeCore/Source/CompositeCore/Private/CompositeCoreSceneViewExtension.cpp
+++ b/Engine/Plugins/Experimental/Compositing/CompositeCore/Source/CompositeCore/Private/CompositeCoreSceneViewExtension.cpp
@@ -17,6 +17,7 @@
 #include "Rendering/CustomRenderPass.h"
 #include "RenderGraphBuilder.h"
 #include "RenderGraphUtils.h"
+#include "SceneRenderTargetParameters.h"
 #include "SceneView.h"
 #include "ScreenPass.h"
 #include "SystemTextures.h"
@@ -28,6 +29,11 @@ static TAutoConsoleVariable<int32> CVarCompositeCoreApplyFXAA(
 	TEXT("When enabled, the custom render pass automatically applies FXAA."),
 	ECVF_RenderThreadSafe);
 
+/** Parameters used to transition a texture to SRVGraphics. */
+BEGIN_SHADER_PARAMETER_STRUCT(FCompositeRenderTargetTransitionParameters, )
+	RDG_TEXTURE_ACCESS(TransitionTexture, ERHIAccess::SRVGraphics)
+END_SHADER_PARAMETER_STRUCT()
+
 class FCompositeCoreCustomRenderPass : public FCustomRenderPassBase
 {
 public:
@@ -119,7 +125,7 @@ void FCompositeCoreSceneViewExtension::UnregisterPrimitives_GameThread(const TAr
 
 void FCompositeCoreSceneViewExtension::SetPostRenderWork_GameThread(UE::CompositeCore::FPostRenderWork&& InWork)
 {
-	bHasCustomRenderWork = !InWork.FramePasses.IsEmpty();
+	bHasCustomRenderWork = !InWork.IsEmpty();
 
 	ENQUEUE_RENDER_COMMAND(CopyCompositeCoreRenderWork)(
 		[InWork = MoveTemp(InWork), WeakThis = AsWeak()](FRHICommandListImmediate& RHICmdList) mutable
@@ -317,6 +323,15 @@ void FCompositeCoreSceneViewExtension::PreRenderViewFamily_RenderThread(FRDGBuil
 
 void FCompositeCoreSceneViewExtension::PreRenderView_RenderThread(FRDGBuilder& GraphBuilder, FSceneView& InView)
 {
+	static uint64 LastFrameCounter = 0;
+
+	// We only apply pre-processing once per frame, on the very first render.
+	if (LastFrameCounter != GFrameCounterRenderThread)
+	{
+		ApplyPreprocessing(GraphBuilder, InView);
+		
+		LastFrameCounter = GFrameCounterRenderThread;
+	}
 }
 
 void FCompositeCoreSceneViewExtension::PrePostProcessPass_RenderThread(FRDGBuilder& GraphBuilder, const FSceneView& InView, const FPostProcessingInputs& Inputs)
@@ -365,7 +380,83 @@ bool FCompositeCoreSceneViewExtension::IsActiveForView(const FSceneView& InView)
 	return bIsActive;
 }
 
-TSortedMap<int32, UE::CompositeCore::FPassInput> FCompositeCoreSceneViewExtension::CreateExternalTextureMap(FRDGBuilder& GraphBuilder, const FSceneView& InView, const FPostProcessMaterialInputs& Inputs) const
+void FCompositeCoreSceneViewExtension::ApplyPreprocessing(FRDGBuilder& GraphBuilder, FSceneView& InView) const
+{
+	using namespace UE::CompositeCore;
+
+	const TSortedMap<ResourceId, TArray<const FCompositeCorePassProxy*>>& PreprocessingPasses = GetRenderWork().PreprocessingPasses;
+
+	if (PreprocessingPasses.IsEmpty())
+	{
+		return;
+	}
+
+	FPassContext PassContext = {};
+	const FExternalTextureMap ExternalTextures = CreateExternalTextureMap(GraphBuilder, InView);
+
+	for (const TPair<ResourceId, TArray<const FCompositeCorePassProxy*>>& Pair : PreprocessingPasses)
+	{
+		const ResourceId ExternalId = Pair.Key;
+		const TArray<const FCompositeCorePassProxy*>& Passes = Pair.Value;
+		const FPassInput* ExternalTexture = ExternalTextures.Find(ExternalId);
+
+		if (!ensureMsgf(ExternalTexture != nullptr, TEXT("Unexpected missing external texture override as output.")))
+		{
+			continue;
+		}
+
+		// Setup pass inputs
+		FPassInputArray PassInputs;
+		PassInputs.GetArray().Add(*ExternalTexture);
+
+		// Iterate over all passes
+		for (int32 PassIndex =0; PassIndex <Passes.Num(); ++PassIndex)
+		{
+			const bool bIsLastPass = (PassIndex == Passes.Num() - 1);
+			const FCompositeCorePassProxy* Pass = Passes[PassIndex];
+
+			FRDGTextureRef ExternalRDG = nullptr;
+			const TOptional<UE::CompositeCore::ResourceId>& OverrrideOutput = Pass->GetDeclareOutputOverride();
+			
+			if (OverrrideOutput.IsSet())
+			{
+				const FExternalRenderTarget* ExternalRT = ExternalInputs_RenderThread.Find(OverrrideOutput.GetValue());
+				if (ensureMsgf(ExternalRT != nullptr, TEXT("Unexpected missing external render target override as output.")))
+				{
+					ExternalRDG = GraphBuilder.RegisterExternalTexture(ExternalRT->RenderTarget);
+					PassInputs.OverrideOutput = FScreenPassRenderTarget(ExternalRDG, ERenderTargetLoadAction::ENoAction);
+
+					// Update view rect to match output
+					PassContext.OutputViewRect = FIntRect(FInt32Point::ZeroValue, ExternalRT->RenderTarget->GetDesc().Extent);
+				}
+			}
+			else
+			{
+				PassInputs.OverrideOutput = FScreenPassRenderTarget{};
+
+				// Update view rect to match external texture
+				PassContext.OutputViewRect = ExternalTexture->Texture.ViewRect;
+			}
+
+			// Register pass and update output
+			const FPassOutput Output = Pass->Add(GraphBuilder, InView, PassInputs, PassContext);
+
+			// Manually transition the external render target back to a shader resource.
+			if (ExternalRDG != nullptr)
+			{
+				FCompositeRenderTargetTransitionParameters* Parameters = GraphBuilder.AllocParameters<FCompositeRenderTargetTransitionParameters>();
+				Parameters->TransitionTexture = ExternalRDG;
+
+				GraphBuilder.AddPass(RDG_EVENT_NAME("CompositeCoreTransitionRT"), Parameters, ERDGPassFlags::Raster | ERDGPassFlags::NeverCull | ERDGPassFlags::SkipRenderPass, [](FRHICommandList&) {});
+			}
+
+			// Update next pass input
+			PassInputs[0] = Output.Resource;
+		}
+	}
+}
+
+TSortedMap<int32, UE::CompositeCore::FPassInput> FCompositeCoreSceneViewExtension::CreateExternalTextureMap(FRDGBuilder& GraphBuilder, const FSceneView& InView) const
 {
 	using namespace UE::CompositeCore;
 
@@ -536,7 +627,7 @@ FScreenPassTexture FCompositeCoreSceneViewExtension::PostProcessWork_RenderThrea
 {
 	using namespace UE::CompositeCore;
 
-	FExternalTextureMap ExternalTextures = CreateExternalTextureMap(GraphBuilder, InView, Inputs);
+	FExternalTextureMap ExternalTextures = CreateExternalTextureMap(GraphBuilder, InView);
 	const TArray<const FCompositeCorePassProxy*>& Passes = GetRenderWork().FramePasses[InLocation];
 
 	FPassContext PassContext;
diff --git a/Engine/Plugins/Experimental/Compositing/CompositeCore/Source/CompositeCore/Private/CompositeCoreSceneViewExtension.h b/Engine/Plugins/Experimental/Compositing/CompositeCore/Source/CompositeCore/Private/CompositeCoreSceneViewExtension.h
--- a/Engine/Plugins/Experimental/Compositing/CompositeCore/Source/CompositeCore/Private/CompositeCoreSceneViewExtension.h
+++ b/Engine/Plugins/Experimental/Compositing/CompositeCore/Source/CompositeCore/Private/CompositeCoreSceneViewExtension.h
@@ -87,8 +87,11 @@ private:
 	/** Active check to avoid work in post-processing. */
 	bool IsActiveForView(const FSceneView& InView) const;
 
+	/** Once-per frame external resource processing passes. */
+	void ApplyPreprocessing(FRDGBuilder& GraphBuilder, FSceneView& InView) const;
+
 	/** Convenience function to create an override map of external textures. */
-	FExternalTextureMap CreateExternalTextureMap(FRDGBuilder& GraphBuilder, const FSceneView& InView, const FPostProcessMaterialInputs& Inputs) const;
+	FExternalTextureMap CreateExternalTextureMap(FRDGBuilder& GraphBuilder, const FSceneView& InView) const;
 	
 	/** Convenience function to update the next pass inputs from the last pass' output. */
 	void UpdateNextPassInputs(
diff --git a/Engine/Plugins/Experimental/Compositing/CompositeCore/Source/CompositeCore/Private/Passes/CompositeCorePassFXAAProxy.cpp b/Engine/Plugins/Experimental/Compositing/CompositeCore/Source/CompositeCore/Private/Passes/CompositeCorePassFXAAProxy.cpp
--- a/Engine/Plugins/Experimental/Compositing/CompositeCore/Source/CompositeCore/Private/Passes/CompositeCorePassFXAAProxy.cpp
+++ b/Engine/Plugins/Experimental/Compositing/CompositeCore/Source/CompositeCore/Private/Passes/CompositeCorePassFXAAProxy.cpp
@@ -34,13 +34,28 @@ namespace UE
 	{
 		namespace Private
 		{
-			FScreenPassTexture AddDisplayTransformPass(FRDGBuilder& GraphBuilder, const FScreenPassViewInfo ViewInfo, const FScreenPassTexture& Input, bool bIsForward, float Gamma = 2.2f)
+			FScreenPassTexture AddDisplayTransformPass(FRDGBuilder& GraphBuilder, const FScreenPassViewInfo ViewInfo, const FScreenPassTexture& Input, TOptional<FScreenPassRenderTarget> OverrideOutput, bool bIsForward, float Gamma = 2.2f)
 			{
 				RDG_EVENT_SCOPE_STAT(GraphBuilder, FCompositeCoreDisplayTransform, "CompositeCore.DisplayTransform");
 				RDG_GPU_STAT_SCOPE(GraphBuilder, FCompositeCoreDisplayTransform);
 
 				FGlobalShaderMap* GlobalShaderMap = GetGlobalShaderMap(ViewInfo.FeatureLevel);
-				FScreenPassRenderTarget Output = FScreenPassRenderTarget::CreateFromInput(GraphBuilder, Input, ERenderTargetLoadAction::ENoAction, TEXT("FXAACompositePass"));
+
+				FScreenPassRenderTarget Output;
+
+				if (OverrideOutput.IsSet())
+				{
+					Output = OverrideOutput.GetValue();
+				}
+				else
+				{
+					FRDGTextureDesc OutputDesc = Input.Texture->Desc;
+					OutputDesc.Flags |= TexCreate_RenderTargetable;
+					OutputDesc.Format = FSceneTexturesConfig::Get().ColorFormat;
+
+					Output = FScreenPassRenderTarget(GraphBuilder.CreateTexture(OutputDesc, TEXT("FXAACompositePass")), ERenderTargetLoadAction::ENoAction);
+				}
+
 				FScreenPassTextureViewport Viewport = FScreenPassTextureViewport(Output);
 
 				FCompositeCoreDisplayTransformShader::FParameters* Parameters = GraphBuilder.AllocParameters<FCompositeCoreDisplayTransformShader::FParameters>();
@@ -77,7 +92,7 @@ namespace UE
 			{
 				// We tonemap & encode the result so that FXAA can operate on perceptual colors
 				constexpr bool bIsForward = true;
-				Input = Private::AddDisplayTransformPass(GraphBuilder, InView, Input, bIsForward);
+				Input = Private::AddDisplayTransformPass(GraphBuilder, InView, Input, {}, bIsForward);
 			}
 
 			static IConsoleVariable* FXAAQualityCVar = IConsoleManager::Get().FindConsoleVariable(TEXT("r.FXAA.Quality"));
@@ -87,7 +102,10 @@ namespace UE
 			FFXAAInputs FXAAInputs{};
 			FXAAInputs.SceneColor = Input;
 			FXAAInputs.Quality = static_cast<EFXAAQuality>(FMath::Clamp(Quality, 0, static_cast<int32>(EFXAAQuality::MAX) - 1));
-			FXAAInputs.OverrideOutput = Inputs.OverrideOutput;
+			if (!bLinearSourceColors)
+			{
+				FXAAInputs.OverrideOutput = Inputs.OverrideOutput;
+			}
 
 			FScreenPassTexture Output = AddFXAAPass(GraphBuilder, InView, FXAAInputs);
 
@@ -95,7 +113,7 @@ namespace UE
 			{
 				// We decode and invert the tonemapping to obtain linear colors again.
 				constexpr bool bIsForward = false;
-				Output = Private::AddDisplayTransformPass(GraphBuilder, InView, Output, bIsForward);
+				Output = Private::AddDisplayTransformPass(GraphBuilder, InView, Output, Inputs.OverrideOutput, bIsForward);
 			}
 
 			return FPassOutput{ MoveTemp(Output), Metadata, PassOutputOverride };
diff --git a/Engine/Plugins/Experimental/Compositing/CompositeCore/Source/CompositeCore/Private/Passes/CompositeCorePassProxy.cpp b/Engine/Plugins/Experimental/Compositing/CompositeCore/Source/CompositeCore/Private/Passes/CompositeCorePassProxy.cpp
--- a/Engine/Plugins/Experimental/Compositing/CompositeCore/Source/CompositeCore/Private/Passes/CompositeCorePassProxy.cpp
+++ b/Engine/Plugins/Experimental/Compositing/CompositeCore/Source/CompositeCore/Private/Passes/CompositeCorePassProxy.cpp
@@ -62,9 +62,28 @@ namespace UE
 		}
 
 		FPostRenderWork::FPostRenderWork()
-			: FramePasses{}
+			: ExternalInputs{}
+			, PreprocessingPasses{}
+			, FramePasses{}
 			, FrameAllocator{new FSceneRenderingBulkObjectAllocator}
 		{ }
+
+		bool FPostRenderWork::IsEmpty() const
+		{
+			bool bIsEmpty = true;
+
+			for (const TPair<ResourceId, TArray<const FCompositeCorePassProxy*>>& Pair : PreprocessingPasses)
+			{
+				bIsEmpty &= Pair.Value.IsEmpty();
+			}
+
+			for (const TPair<ISceneViewExtension::EPostProcessingPass, TArray<const FCompositeCorePassProxy*>>& Pair : FramePasses)
+			{
+				bIsEmpty &= Pair.Value.IsEmpty();
+			}
+
+			return bIsEmpty;
+		}
 	}
 }
 
diff --git a/Engine/Plugins/Experimental/Compositing/CompositeCore/Source/CompositeCore/Public/Passes/CompositeCorePassProxy.h b/Engine/Plugins/Experimental/Compositing/CompositeCore/Source/CompositeCore/Public/Passes/CompositeCorePassProxy.h
--- a/Engine/Plugins/Experimental/Compositing/CompositeCore/Source/CompositeCore/Public/Passes/CompositeCorePassProxy.h
+++ b/Engine/Plugins/Experimental/Compositing/CompositeCore/Source/CompositeCore/Public/Passes/CompositeCorePassProxy.h
@@ -231,9 +231,15 @@ namespace UE
 			/** Constructor. */
 			UE_API FPostRenderWork();
 
+			/** Returns true if there are no active work passes. */
+			bool IsEmpty() const;
+
 			/** Array of user-defined external input texture overrides, where indices map to ResourceId. */
 			TArray<FExternalTexture> ExternalInputs;
 
+			/** Array of pre-processing passes applied on specific texture resources at the start of rendering. */
+			TSortedMap<ResourceId, TArray<const FCompositeCorePassProxy*>> PreprocessingPasses;
+
 			/** Post-processing passes at the specified locations. */
 			TSortedMap<ISceneViewExtension::EPostProcessingPass, TArray<const FCompositeCorePassProxy*>> FramePasses;
 
@@ -279,6 +285,24 @@ public:
 		return PassDeclaredInputs[InputIndex];
 	}
 
+	/** Get declared output pass override. */
+	const TOptional<UE::CompositeCore::ResourceId>& GetDeclareOutputOverride() const
+	{
+		return PassOutputOverride;	
+	}
+
+	/** Declare output pass override. */
+	void DeclareOutputOverride(UE::CompositeCore::ResourceId InExternalId)
+	{
+		PassOutputOverride = InExternalId;
+	}
+
+	/** Reset output pass override. */
+	void ResetOutputOverride()
+	{
+		PassOutputOverride.Reset();
+	}
+
 	/** Return sub-passes per input index. */
 	const TArray<const FCompositeCorePassProxy*>* GetSubPasses(int32 InputIndex) const
 	{
