CompositeCore: Avoid validation pop-up on valid settings.

Note that the condition causing this setting/cvar mistmatch will be fixed in UE-291767.

diff --git a/Engine/Plugins/Experimental/Compositing/CompositeCore/Source/CompositeCore/Private/CompositeCoreSubsystem.cpp b/Engine/Plugins/Experimental/Compositing/CompositeCore/Source/CompositeCore/Private/CompositeCoreSubsystem.cpp
--- a/Engine/Plugins/Experimental/Compositing/CompositeCore/Source/CompositeCore/Private/CompositeCoreSubsystem.cpp
+++ b/Engine/Plugins/Experimental/Compositing/CompositeCore/Source/CompositeCore/Private/CompositeCoreSubsystem.cpp
@@ -155,11 +155,8 @@ bool UCompositeCoreSubsystem::ValidateProjectSettings()
 	if (!bValidSettings )
 	{
 #if WITH_EDITOR
-		URendererSettings* RendererSettings = GetMutableDefault<URendererSettings>();
-		check(RendererSettings);
-
 		// We inform the user and offer them the option to activate project settings.
-		UE_CALL_ONCE([&]{ PrimitiveHoldoutSettingsNotification(RendererSettings); });
+		UE_CALL_ONCE([&]{ PrimitiveHoldoutSettingsNotification(); });
 #else
 		UE_CALL_ONCE([&]{ UE_LOG(LogCompositeCore, Warning, TEXT("Both \"Alpha Output\" and \"Support Primitive Alpha Holdout\" project settings must be enabled for holdout composite.")); });
 #endif
@@ -169,11 +166,20 @@ bool UCompositeCoreSubsystem::ValidateProjectSettings()
 }
 
 #if WITH_EDITOR
-void UCompositeCoreSubsystem::PrimitiveHoldoutSettingsNotification(URendererSettings* RendererSettings)
+void UCompositeCoreSubsystem::PrimitiveHoldoutSettingsNotification()
 {
+	URendererSettings* RendererSettings = GetMutableDefault<URendererSettings>();
+	check(RendererSettings);
+
 	const bool bAlphaOutputMissing = !RendererSettings->bEnableAlphaChannelInPostProcessing;
 	const bool bPrimitiveHoldoutMissing = !RendererSettings->bDeferredSupportPrimitiveAlphaHoldout;
 
+	if (!bAlphaOutputMissing && !bPrimitiveHoldoutMissing)
+	{
+		UE_LOG(LogCompositeCore, Warning, TEXT("\"Alpha Output\" and \"Support Primitive Alpha Holdout\" project settings are valid, but one or more console variables have been disabled."));
+		return;
+	}
+
 	const FText AlphaOutputSettingOption = LOCTEXT("HoldoutSetting_AlphaOutput", "\n- Alpha Output");
 	const FText PrimitiveHoldoutSettingOption = LOCTEXT("HoldoutSetting_PrimitiveHoldout", "\n- Support Primitive Alpha Holdout");
 	const FText HoldoutText = FText::Format(LOCTEXT("HoldoutSettingPrompt", "The following project setting(s) must be enabled for holdout composite:{0}{1}\n\nWarning: update can add renderer performance costs."), bAlphaOutputMissing ? AlphaOutputSettingOption : FText::GetEmpty(), bPrimitiveHoldoutMissing ? PrimitiveHoldoutSettingOption : FText::GetEmpty());
diff --git a/Engine/Plugins/Experimental/Compositing/CompositeCore/Source/CompositeCore/Public/CompositeCoreSubsystem.h b/Engine/Plugins/Experimental/Compositing/CompositeCore/Source/CompositeCore/Public/CompositeCoreSubsystem.h
--- a/Engine/Plugins/Experimental/Compositing/CompositeCore/Source/CompositeCore/Public/CompositeCoreSubsystem.h
+++ b/Engine/Plugins/Experimental/Compositing/CompositeCore/Source/CompositeCore/Public/CompositeCoreSubsystem.h
@@ -13,7 +13,6 @@ class FCompositeCoreSceneViewExtension;
 struct FCompositeCoreSettings;
 class SNotificationItem;
 class UPrimitiveComponent;
-class URendererSettings;
 
 namespace UE::CompositeCore
 {
@@ -80,7 +79,7 @@ private:
 
 #if WITH_EDITOR
 	/* Toast notification to ask users to enable the missing project settings. */
-	UE_API void PrimitiveHoldoutSettingsNotification(URendererSettings* RendererSettings);
+	UE_API void PrimitiveHoldoutSettingsNotification();
 
 	/* Toast notification item. */
 	TWeakPtr<SNotificationItem> HoldoutNotificationItem;
