CompositeCore: Rename FPostRenderWork to FRenderWork, since as of CL 43121477 scene view extension work can also occur before post-processing.

diff --git a/Engine/Plugins/Experimental/Compositing/Composite/Source/Composite/Private/CompositeActor.cpp b/Engine/Plugins/Experimental/Compositing/Composite/Source/Composite/Private/CompositeActor.cpp
--- a/Engine/Plugins/Experimental/Compositing/Composite/Source/Composite/Private/CompositeActor.cpp
+++ b/Engine/Plugins/Experimental/Compositing/Composite/Source/Composite/Private/CompositeActor.cpp
@@ -104,15 +104,15 @@ namespace UE
 				}
 			}
 
-			/* Get post-processing render work (per frame) given the current composite actor properties. */
-			UE::CompositeCore::FPostRenderWork GetPostRenderWork(const ACompositeActor& InCompositeActor)
+			/* Gather render work (per frame) given the current composite actor properties. */
+			UE::CompositeCore::FRenderWork GatherFrameRenderWork(const ACompositeActor& InCompositeActor)
 			{
 				if (InCompositeActor.CompositeLayers.IsEmpty())
 				{
 					return {};
 				}
 
-				UE::CompositeCore::FPostRenderWork RenderWork;
+				UE::CompositeCore::FRenderWork RenderWork;
 				TArray<TObjectPtr<UCompositeLayerBase>> Layers;
 
 				TObjectPtr<UCompositeLayerBase> SoloLayer = FindSoloLayer(InCompositeActor.CompositeLayers);
@@ -248,7 +248,7 @@ void ACompositeActor::Tick(float DeltaTime)
 	UCompositeCoreSubsystem* Subsystem = UWorld::GetSubsystem<UCompositeCoreSubsystem>(GetWorld());
 	if (IsValid(Subsystem))
 	{
-		Subsystem->SetPostRenderWork(UE::Composite::Private::GetPostRenderWork(*this));
+		Subsystem->SetRenderWork(UE::Composite::Private::GatherFrameRenderWork(*this));
 
 		UE::CompositeCore::FBuiltInRenderPassOptions Options;
 		if (bOverridesViewUserFlags)
@@ -572,7 +572,7 @@ void ACompositeActor::RemoveAllState()
 	
 	if (IsValid(Subsystem))
 	{
-		Subsystem->ResetPostRenderWork();
+		Subsystem->ResetRenderWork();
 		Subsystem->ResetBuiltInRenderPassOptions();
 	}
 
diff --git a/Engine/Plugins/Experimental/Compositing/CompositeCore/Source/CompositeCore/Private/CompositeCoreSceneViewExtension.cpp b/Engine/Plugins/Experimental/Compositing/CompositeCore/Source/CompositeCore/Private/CompositeCoreSceneViewExtension.cpp
--- a/Engine/Plugins/Experimental/Compositing/CompositeCore/Source/CompositeCore/Private/CompositeCoreSceneViewExtension.cpp
+++ b/Engine/Plugins/Experimental/Compositing/CompositeCore/Source/CompositeCore/Private/CompositeCoreSceneViewExtension.cpp
@@ -123,7 +123,7 @@ void FCompositeCoreSceneViewExtension::UnregisterPrimitives_GameThread(const TAr
 	}
 }
 
-void FCompositeCoreSceneViewExtension::SetPostRenderWork_GameThread(UE::CompositeCore::FPostRenderWork&& InWork)
+void FCompositeCoreSceneViewExtension::SetRenderWork_GameThread(UE::CompositeCore::FRenderWork&& InWork)
 {
 	bHasCustomRenderWork = !InWork.IsEmpty();
 
@@ -153,12 +153,12 @@ void FCompositeCoreSceneViewExtension::SetPostRenderWork_GameThread(UE::Composit
 					}
 				}
 
-				SVE->PostRenderWork_RenderThread = MoveTemp(InWork);
+				SVE->RenderWork_RenderThread = MoveTemp(InWork);
 			}
 		});
 }
 
-void FCompositeCoreSceneViewExtension::ResetPostRenderWork_GameThread()
+void FCompositeCoreSceneViewExtension::ResetRenderWork_GameThread()
 {
 	bHasCustomRenderWork = false;
 
@@ -168,7 +168,7 @@ void FCompositeCoreSceneViewExtension::ResetPostRenderWork_GameThread()
 			TSharedPtr<FCompositeCoreSceneViewExtension> SVE = StaticCastSharedPtr<FCompositeCoreSceneViewExtension>(WeakThis.Pin());
 			if (SVE.IsValid())
 			{
-				SVE->PostRenderWork_RenderThread.Reset();
+				SVE->RenderWork_RenderThread.Reset();
 				SVE->ExternalInputs_RenderThread.Reset();
 			}
 		});
@@ -208,15 +208,15 @@ int32 FCompositeCoreSceneViewExtension::GetPriority() const
 	return GetDefault<UCompositeCorePluginSettings>()->SceneViewExtensionPriority;
 }
 
-const UE::CompositeCore::FPostRenderWork& FCompositeCoreSceneViewExtension::GetRenderWork() const
+const UE::CompositeCore::FRenderWork& FCompositeCoreSceneViewExtension::GetRenderWork() const
 {
-	if (PostRenderWork_RenderThread.IsSet())
+	if (RenderWork_RenderThread.IsSet())
 	{
-		return PostRenderWork_RenderThread.GetValue();
+		return RenderWork_RenderThread.GetValue();
 	}
 	else
 	{
-		return UE::CompositeCore::FPostRenderWork::GetDefault();
+		return UE::CompositeCore::FRenderWork::GetDefault();
 	}
 }
 
diff --git a/Engine/Plugins/Experimental/Compositing/CompositeCore/Source/CompositeCore/Private/CompositeCoreSceneViewExtension.h b/Engine/Plugins/Experimental/Compositing/CompositeCore/Source/CompositeCore/Private/CompositeCoreSceneViewExtension.h
--- a/Engine/Plugins/Experimental/Compositing/CompositeCore/Source/CompositeCore/Private/CompositeCoreSceneViewExtension.h
+++ b/Engine/Plugins/Experimental/Compositing/CompositeCore/Source/CompositeCore/Private/CompositeCoreSceneViewExtension.h
@@ -45,11 +45,11 @@ public:
 	/** Unregister primitives for compositing. */
 	void UnregisterPrimitives_GameThread(const TArray<UPrimitiveComponent*>& InPrimitiveComponents);
 
-	/** Set post-processing render work to be enqueued on the render thread. */
-	void SetPostRenderWork_GameThread(UE::CompositeCore::FPostRenderWork&& InWork);
+	/** Set frame render work to be enqueued on the render thread. */
+	void SetRenderWork_GameThread(UE::CompositeCore::FRenderWork&& InWork);
 
-	/** Reset post-processing render work. */
-	void ResetPostRenderWork_GameThread();
+	/** Reset frame render work. */
+	void ResetRenderWork_GameThread();
 
 	/** Set options for built-in composite custom render pass. */
 	void SetBuiltInRenderPassOptions_GameThread(const UE::CompositeCore::FBuiltInRenderPassOptions& InOptions);
@@ -82,7 +82,7 @@ private:
 	using FExternalTextureMap = TSortedMap<int32, UE::CompositeCore::FPassInput>;
 
 	/** Getter for the frame render work, either from the actor or default. */
-	const UE::CompositeCore::FPostRenderWork& GetRenderWork() const;
+	const UE::CompositeCore::FRenderWork& GetRenderWork() const;
 
 	/** Active check to avoid work in post-processing. */
 	bool IsActiveForView(const FSceneView& InView) const;
@@ -129,7 +129,7 @@ private:
 	TMap<uint32, TRefCountPtr<IPooledRenderTarget>> CustomRenderTargetPerView_RenderThread;
 
 	/** Custom render work. */
-	TOptional<UE::CompositeCore::FPostRenderWork> PostRenderWork_RenderThread;
+	TOptional<UE::CompositeCore::FRenderWork> RenderWork_RenderThread;
 
 	/** Built-in composite render pass options.*/
 	TOptional<UE::CompositeCore::FBuiltInRenderPassOptions> BuiltInRenderPassOptions;
diff --git a/Engine/Plugins/Experimental/Compositing/CompositeCore/Source/CompositeCore/Private/CompositeCoreSubsystem.cpp b/Engine/Plugins/Experimental/Compositing/CompositeCore/Source/CompositeCore/Private/CompositeCoreSubsystem.cpp
--- a/Engine/Plugins/Experimental/Compositing/CompositeCore/Source/CompositeCore/Private/CompositeCoreSubsystem.cpp
+++ b/Engine/Plugins/Experimental/Compositing/CompositeCore/Source/CompositeCore/Private/CompositeCoreSubsystem.cpp
@@ -105,19 +105,19 @@ void UCompositeCoreSubsystem::UnregisterPrimitives(const TArray<UPrimitiveCompon
 	}
 }
 
-void UCompositeCoreSubsystem::SetPostRenderWork(UE::CompositeCore::FPostRenderWork&& InWork)
+void UCompositeCoreSubsystem::SetRenderWork(UE::CompositeCore::FRenderWork&& InWork)
 {
 	if (CompositeCoreViewExtension.IsValid())
 	{
-		CompositeCoreViewExtension->SetPostRenderWork_GameThread(MoveTemp(InWork));
+		CompositeCoreViewExtension->SetRenderWork_GameThread(MoveTemp(InWork));
 	}
 }
 
-void UCompositeCoreSubsystem::ResetPostRenderWork()
+void UCompositeCoreSubsystem::ResetRenderWork()
 {
 	if (CompositeCoreViewExtension.IsValid())
 	{
-		CompositeCoreViewExtension->ResetPostRenderWork_GameThread();
+		CompositeCoreViewExtension->ResetRenderWork_GameThread();
 	}
 }
 
diff --git a/Engine/Plugins/Experimental/Compositing/CompositeCore/Source/CompositeCore/Private/Passes/CompositeCorePassProxy.cpp b/Engine/Plugins/Experimental/Compositing/CompositeCore/Source/CompositeCore/Private/Passes/CompositeCorePassProxy.cpp
--- a/Engine/Plugins/Experimental/Compositing/CompositeCore/Source/CompositeCore/Private/Passes/CompositeCorePassProxy.cpp
+++ b/Engine/Plugins/Experimental/Compositing/CompositeCore/Source/CompositeCore/Private/Passes/CompositeCorePassProxy.cpp
@@ -34,11 +34,11 @@ namespace UE
 		{
 		}
 
-		const FPostRenderWork& FPostRenderWork::GetDefault()
+		const FRenderWork& FRenderWork::GetDefault()
 		{
-			static FPostRenderWork DefaultRenderWork = []() -> FPostRenderWork
+			static FRenderWork DefaultRenderWork = []() -> FRenderWork
 				{
-					FPostRenderWork RenderWork;
+					FRenderWork RenderWork;
 					TArray<const FCompositeCorePassProxy*>& SSRInputWork = RenderWork.FramePasses.Add(ISceneViewExtension::EPostProcessingPass::SSRInput);
 					TArray<const FCompositeCorePassProxy*>& AfterTonemapWork = RenderWork.FramePasses.Add(ISceneViewExtension::EPostProcessingPass::Tonemap);
 
@@ -61,14 +61,14 @@ namespace UE
 			return DefaultRenderWork;
 		}
 
-		FPostRenderWork::FPostRenderWork()
+		FRenderWork::FRenderWork()
 			: ExternalInputs{}
 			, PreprocessingPasses{}
 			, FramePasses{}
 			, FrameAllocator{new FSceneRenderingBulkObjectAllocator}
 		{ }
 
-		bool FPostRenderWork::IsEmpty() const
+		bool FRenderWork::IsEmpty() const
 		{
 			bool bIsEmpty = true;
 
diff --git a/Engine/Plugins/Experimental/Compositing/CompositeCore/Source/CompositeCore/Public/CompositeCoreSubsystem.h b/Engine/Plugins/Experimental/Compositing/CompositeCore/Source/CompositeCore/Public/CompositeCoreSubsystem.h
--- a/Engine/Plugins/Experimental/Compositing/CompositeCore/Source/CompositeCore/Public/CompositeCoreSubsystem.h
+++ b/Engine/Plugins/Experimental/Compositing/CompositeCore/Source/CompositeCore/Public/CompositeCoreSubsystem.h
@@ -17,7 +17,7 @@ class UPrimitiveComponent;
 namespace UE::CompositeCore
 {
 	struct FBuiltInRenderPassOptions;
-	struct FPostRenderWork;
+	struct FRenderWork;
 }
 
 /**
@@ -57,11 +57,11 @@ public:
 	/* Unregister multiple primitives from compositing. */
 	UE_API void UnregisterPrimitives(const TArray<UPrimitiveComponent*>& InPrimitiveComponents);
 
-	/* Set post-processing frame render work. */
-	UE_API void SetPostRenderWork(UE::CompositeCore::FPostRenderWork&& InWork);
+	/* Set scene view extension frame render work. */
+	UE_API void SetRenderWork(UE::CompositeCore::FRenderWork&& InWork);
 
-	/* Reset post-processing frame render work. */
-	UE_API void ResetPostRenderWork();
+	/* Reset scene view extension frame render work. */
+	UE_API void ResetRenderWork();
 
 	/* Set built-in composite render pass options. */
 	UE_API void SetBuiltInRenderPassOptions(const UE::CompositeCore::FBuiltInRenderPassOptions& InOptions);
diff --git a/Engine/Plugins/Experimental/Compositing/CompositeCore/Source/CompositeCore/Public/Passes/CompositeCorePassProxy.h b/Engine/Plugins/Experimental/Compositing/CompositeCore/Source/CompositeCore/Public/Passes/CompositeCorePassProxy.h
--- a/Engine/Plugins/Experimental/Compositing/CompositeCore/Source/CompositeCore/Public/Passes/CompositeCorePassProxy.h
+++ b/Engine/Plugins/Experimental/Compositing/CompositeCore/Source/CompositeCore/Public/Passes/CompositeCorePassProxy.h
@@ -44,7 +44,7 @@ namespace UE
 		static constexpr ResourceId BUILT_IN_CRP_ID = 1;
 		/** Built-in empty/black identifier. */
 		static constexpr ResourceId BUILT_IN_EMPTY_ID = 2;
-		/** First identifier of the external texture inputs, see ExternalInputs on FPostRenderWork. */
+		/** First identifier of the external texture inputs, see ExternalInputs on FRenderWork. */
 		static constexpr ResourceId EXTERNAL_RANGE_START_ID = 100;
 
 		/** Texture resource metadata. */
@@ -222,14 +222,14 @@ namespace UE
 			bool bOpacifyOutput;
 		};
 
-		/** Render-thread struct for render post-processing work per frame. */
-		struct FPostRenderWork final
+		/** Render-thread struct for scene view extension render work per frame. */
+		struct FRenderWork final
 		{
 			/** Default frame work. */
-			static UE_API const FPostRenderWork& GetDefault();
+			static UE_API const FRenderWork& GetDefault();
 
 			/** Constructor. */
-			UE_API FPostRenderWork();
+			UE_API FRenderWork();
 
 			/** Returns true if there are no active work passes. */
 			bool IsEmpty() const;
