CompositeCore: Added new cvar to restrict (deferred) primitve alpha holdout to VMI_Lit only.

Disabled by default to avoid any breaking change, especially if movie render graph has use case for holdout in various view modes.

diff --git a/Engine/Plugins/Experimental/Compositing/CompositeCore/Source/CompositeCore/Private/CompositeCoreSceneViewExtension.cpp b/Engine/Plugins/Experimental/Compositing/CompositeCore/Source/CompositeCore/Private/CompositeCoreSceneViewExtension.cpp
--- a/Engine/Plugins/Experimental/Compositing/CompositeCore/Source/CompositeCore/Private/CompositeCoreSceneViewExtension.cpp
+++ b/Engine/Plugins/Experimental/Compositing/CompositeCore/Source/CompositeCore/Private/CompositeCoreSceneViewExtension.cpp
@@ -29,6 +29,13 @@ static TAutoConsoleVariable<int32> CVarCompositeCoreApplyFXAA(
 	TEXT("When enabled, the custom render pass automatically applies FXAA."),
 	ECVF_RenderThreadSafe);
 
+static TAutoConsoleVariable<int32> CVarCompositeCoreRestrictedHoldoutViewMode(
+	TEXT("CompositeCore.RestrictedHoldoutViewMode"),
+	0,
+	TEXT(" 0: no restrictions (default)\n")
+	TEXT(" 1: primitive alpha holdout is restricted to the Lit view mode in deferred rendering.")
+);
+
 /** Parameters used to transition a texture to SRVGraphics. */
 BEGIN_SHADER_PARAMETER_STRUCT(FCompositeRenderTargetTransitionParameters, )
 	RDG_TEXTURE_ACCESS(TransitionTexture, ERHIAccess::SRVGraphics)
@@ -238,6 +245,16 @@ void FCompositeCoreSceneViewExtension::SetupView(FSceneViewFamily& InViewFamily,
 
 void FCompositeCoreSceneViewExtension::BeginRenderViewFamily(FSceneViewFamily& InViewFamily)
 {
+	// Optionally disable primitive alpha holdout on non-lit view modes (except path tracing)
+	// for easier workflows in editor with multiple viewports.
+	const bool bAllowPrimitiveAlphaHoldout = InViewFamily.EngineShowFlags.AllowPrimitiveAlphaHoldout;
+	const int32 bLitOnlyPrimitiveAlphaHoldout = CVarCompositeCoreRestrictedHoldoutViewMode.GetValueOnGameThread() > 0;
+
+	if (bAllowPrimitiveAlphaHoldout && bLitOnlyPrimitiveAlphaHoldout && InViewFamily.ViewMode != VMI_PathTracing)
+	{
+		InViewFamily.EngineShowFlags.SetAllowPrimitiveAlphaHoldout(InViewFamily.ViewMode == VMI_Lit);
+	}
+
 	// Disable custom render pass for view families without primitive alpha holdout support.
 	if (InViewFamily.EngineShowFlags.AllowPrimitiveAlphaHoldout == 0)
 	{
