CompositeCore: Fix recent FXAA proxy pass issue introduced by pre-processing support at CL 43121493.

* Optional parameter was not a valid check on the present of a valid output override texture.
* Re-allow pass on Composite Actor preprocessing.

diff --git a/Engine/Plugins/Experimental/Compositing/Composite/Source/Composite/Public/Layers/CompositeLayerPlate.h b/Engine/Plugins/Experimental/Compositing/Composite/Source/Composite/Public/Layers/CompositeLayerPlate.h
--- a/Engine/Plugins/Experimental/Compositing/Composite/Source/Composite/Public/Layers/CompositeLayerPlate.h
+++ b/Engine/Plugins/Experimental/Compositing/Composite/Source/Composite/Public/Layers/CompositeLayerPlate.h
@@ -96,7 +96,7 @@ public:
 	TObjectPtr<UTexture> Texture;
 
 	/** Texture pre-processing passes, applied at the start of rendering and updating a render target. */
-	UPROPERTY(EditAnywhere, BlueprintReadWrite, Instanced, Category = "", meta = (DisplayName = "Pre-Processing", DisallowedClasses = "/Script/Composite.CompositePassMaterial, /Script/Composite.CompositePassFXAA"))
+	UPROPERTY(EditAnywhere, BlueprintReadWrite, Instanced, Category = "", meta = (DisplayName = "Pre-Processing", DisallowedClasses = "/Script/Composite.CompositePassMaterial"))
 	TArray<TObjectPtr<UCompositePassBase>> Preprocessing;
 
 	/** Sub-passes applied during post-processing on the texture or holdout composite render. */
diff --git a/Engine/Plugins/Experimental/Compositing/CompositeCore/Source/CompositeCore/Private/Passes/CompositeCorePassFXAAProxy.cpp b/Engine/Plugins/Experimental/Compositing/CompositeCore/Source/CompositeCore/Private/Passes/CompositeCorePassFXAAProxy.cpp
--- a/Engine/Plugins/Experimental/Compositing/CompositeCore/Source/CompositeCore/Private/Passes/CompositeCorePassFXAAProxy.cpp
+++ b/Engine/Plugins/Experimental/Compositing/CompositeCore/Source/CompositeCore/Private/Passes/CompositeCorePassFXAAProxy.cpp
@@ -35,20 +35,16 @@ namespace UE
 	{
 		namespace Private
 		{
-			FScreenPassTexture AddDisplayTransformPass(FRDGBuilder& GraphBuilder, const FScreenPassViewInfo ViewInfo, const FScreenPassTexture& Input, TOptional<FScreenPassRenderTarget> OverrideOutput, bool bIsForward, float Gamma = 2.2f)
+			FScreenPassTexture AddDisplayTransformPass(FRDGBuilder& GraphBuilder, const FScreenPassViewInfo ViewInfo, const FScreenPassTexture& Input, const FScreenPassRenderTarget& OverrideOutput, bool bIsForward, float Gamma = 2.2f)
 			{
 				RDG_EVENT_SCOPE_STAT(GraphBuilder, FCompositeCoreDisplayTransform, "CompositeCore.DisplayTransform");
 				RDG_GPU_STAT_SCOPE(GraphBuilder, FCompositeCoreDisplayTransform);
 
 				FGlobalShaderMap* GlobalShaderMap = GetGlobalShaderMap(ViewInfo.FeatureLevel);
 
-				FScreenPassRenderTarget Output;
+				FScreenPassRenderTarget Output = OverrideOutput;
 
-				if (OverrideOutput.IsSet())
-				{
-					Output = OverrideOutput.GetValue();
-				}
-				else
+				if (!Output.IsValid())
 				{
 					FRDGTextureDesc OutputDesc = Input.Texture->Desc;
 					OutputDesc.Flags |= TexCreate_RenderTargetable;
