CompositeCore: Improved dilation pass.

The previous dilation logic incorrectly kept iterating as it found increasingly more opaque colors, which were usually further away from the edge. This caused incoherent dilated colors.

#prelight 685c79f7eb2c0fb1019e6278

diff --git a/Engine/Plugins/Experimental/Compositing/CompositeCore/Shaders/Private/CompositeCoreDilate.usf b/Engine/Plugins/Experimental/Compositing/CompositeCore/Shaders/Private/CompositeCoreDilate.usf
--- a/Engine/Plugins/Experimental/Compositing/CompositeCore/Shaders/Private/CompositeCoreDilate.usf
+++ b/Engine/Plugins/Experimental/Compositing/CompositeCore/Shaders/Private/CompositeCoreDilate.usf
@@ -38,13 +38,17 @@ static const int2 kOffsets[KERNEL_COUNT] = {
 };
 #elif DILATION_SIZE == 2
 static const int2 kOffsets[24] = {
-    int2(-2,  0),
     int2(-1,  0),
     int2( 1,  0),
+    int2( 0,  1),
+    int2( 0, -1),
+    int2(-1,  1),
+    int2( 1,  1),
+    int2( 1, -1),
+    int2(-1, -1),
+    int2(-2,  0),
     int2( 2,  0),
     int2( 0, -2),
-    int2( 0, -1),
-    int2( 0,  1),
     int2( 0,  2),
     int2(-2, -1),
     int2(-1, -2),
@@ -54,10 +58,6 @@ static const int2 kOffsets[24] = {
     int2( 1,  2),
     int2(-1,  2),
     int2(-2,  1),
-    int2(-1, -1),
-    int2( 1, -1),
-    int2( 1,  1),
-    int2(-1,  1),
     int2(-2, -2),
     int2( 2, -2),
     int2( 2,  2),
@@ -152,6 +152,7 @@ void MainCS(
             if(OffsetColor.a < TexelColor.a)
             {
                 TexelColor = OffsetColor;
+                break; // We break as soon as we find a more opaque color
             }
         }
     }
