CompositeCore: Fix CIS division by zero warnings.

diff --git a/Engine/Plugins/Experimental/Compositing/CompositeCore/Shaders/Private/CompositeCoreBlending.ush b/Engine/Plugins/Experimental/Compositing/CompositeCore/Shaders/Private/CompositeCoreBlending.ush
--- a/Engine/Plugins/Experimental/Compositing/CompositeCore/Shaders/Private/CompositeCoreBlending.ush
+++ b/Engine/Plugins/Experimental/Compositing/CompositeCore/Shaders/Private/CompositeCoreBlending.ush
@@ -98,7 +98,7 @@ float ColorDodgeBlend(float Src, float Dst)
 	}
 	else if(Dst > 0 && Dst < 1)
 	{
-		return min(1, Dst / (1 - Src));
+		return select(Src < 1, min(1, Dst / (1 - Src)), 1.0f);
 	}
 	else
 	{
@@ -132,7 +132,7 @@ float ColorBurnBlend(float Src, float Dst)
 	}
 	else if(Dst > 0 && Dst < 1)
 	{
-		return 1 - min(1, (1 - Dst) / Src);
+		return select(Src > 0, 1 - min(1, (1 - Dst) / Src), 0.0f);
 	}
 	else
 	{
