CompositeCore: Remove subpass interface in favor of more general proxy (tree/graph) traversal.

* Pass proxies can now declare another proxy as an input type, whose output will feed the specified input index.
* Traversal is done recursively in the scene view extension.

diff --git a/Engine/Plugins/Experimental/Compositing/Composite/Source/Composite/Private/Layers/CompositeLayerMainRender.cpp b/Engine/Plugins/Experimental/Compositing/Composite/Source/Composite/Private/Layers/CompositeLayerMainRender.cpp
--- a/Engine/Plugins/Experimental/Compositing/Composite/Source/Composite/Private/Layers/CompositeLayerMainRender.cpp
+++ b/Engine/Plugins/Experimental/Compositing/Composite/Source/Composite/Private/Layers/CompositeLayerMainRender.cpp
@@ -17,7 +17,7 @@ bool UCompositeLayerMainRender::GetProxy(FTraversalContext& InContext, FSceneRen
 	using namespace UE::CompositeCore;
 
 	FPassInputDeclArray Inputs;
-	Inputs.SetNum(2);
+	Inputs.SetNum(FixedNumLayerInputs);
 
 	FPassInternalResourceDesc Desc0;
 	Desc0.Index = 0;
diff --git a/Engine/Plugins/Experimental/Compositing/Composite/Source/Composite/Private/Layers/CompositeLayerPlate.cpp b/Engine/Plugins/Experimental/Compositing/Composite/Source/Composite/Private/Layers/CompositeLayerPlate.cpp
--- a/Engine/Plugins/Experimental/Compositing/Composite/Source/Composite/Private/Layers/CompositeLayerPlate.cpp
+++ b/Engine/Plugins/Experimental/Compositing/Composite/Source/Composite/Private/Layers/CompositeLayerPlate.cpp
@@ -393,9 +393,6 @@ bool UCompositeLayerPlate::GetProxy(FTraversalContext& InContext, FSceneRenderin
 {
 	using namespace UE::CompositeCore;
 
-	FPassInputDeclArray Inputs;
-	Inputs.SetNum(2);
-
 	TOptional<ResourceId> TextureId;
 	TOptional<ResourceId> TextureProcessedId;
 
@@ -485,54 +482,34 @@ bool UCompositeLayerPlate::GetProxy(FTraversalContext& InContext, FSceneRenderin
 		InputId = TextureId.GetValue();
 	}
 
-	Inputs[0].Set<FPassExternalResourceDesc>({ InputId });
-	Inputs[1] = GetDefaultSecondInput(InContext);
-
-	// Copy for subpasses before we move the declared inputs
-	FPassInputDecl BaseInput0 = Inputs[0];
+	FPassInputDecl PassInput;
+	PassInput.Set<FPassExternalResourceDesc>({ InputId });
 
-	FMergePassProxy* Proxy = InFrameAllocator.Create<FMergePassProxy>(MoveTemp(Inputs));
-	Proxy->MergeOp = GetOperation(InContext);
-	Proxy->bUseLensDistortion = false;
-
-	if (!Postprocessing.IsEmpty())
+	// Iterate backwards so that lower passes in the UI are executed first
+	for (int32 PassIndex = Postprocessing.Num() - 1; PassIndex >= 0; --PassIndex)
 	{
-		TArray<const FCompositeCorePassProxy*> SubPassProxies;
-		SubPassProxies.Reserve(Postprocessing.Num());
-
-		bool bFirstValidPass = true;
+		const TObjectPtr<UCompositePassBase>& SubPass = Postprocessing[PassIndex];
 
-		// Iterate backwards so that lower passes in the UI are executed first
-		for (int32 PassIndex = Postprocessing.Num() - 1; PassIndex >= 0; --PassIndex)
+		if (IsValid(SubPass) && SubPass->IsActive())
 		{
-			const TObjectPtr<UCompositePassBase>& SubPass = Postprocessing[PassIndex];
-
-			if (IsValid(SubPass) && SubPass->IsActive())
+			FCompositeCorePassProxy* ChildPassProxy;
+			if (SubPass->GetProxy(PassInput, InFrameAllocator, ChildPassProxy))
 			{
-				FPassInputDecl PassInput;
-				if (bFirstValidPass)
-				{
-					PassInput = BaseInput0;
-					bFirstValidPass = false;
-				}
-				else
-				{
-					// Default input from previous pass
-					PassInput.Set<FPassInternalResourceDesc>({});
-				}
-
-				FCompositeCorePassProxy* SubPassProxy;
-				if (SubPass->GetProxy(PassInput, InFrameAllocator, SubPassProxy))
-				{
-					SubPassProxies.Add(SubPassProxy);
-				}
+				// Next pass should receive the output of the current child pass.
+				PassInput.Set<const FCompositeCorePassProxy*>(ChildPassProxy);
 			}
 		}
-
-		// Apply sub-passes overriding the first input
-		Proxy->AddSubPasses(0, SubPassProxies);
 	}
 
+	FPassInputDeclArray Inputs;
+	Inputs.SetNum(FixedNumLayerInputs);
+	Inputs[0] = PassInput;
+	Inputs[1] = GetDefaultSecondInput(InContext);
+
+	FMergePassProxy* Proxy = InFrameAllocator.Create<FMergePassProxy>(MoveTemp(Inputs));
+	Proxy->MergeOp = GetOperation(InContext);
+	Proxy->bUseLensDistortion = false;
+
 	OutProxy = Proxy;
 	return true;
 }
diff --git a/Engine/Plugins/Experimental/Compositing/Composite/Source/Composite/Private/Layers/CompositeLayerSceneCapture.cpp b/Engine/Plugins/Experimental/Compositing/Composite/Source/Composite/Private/Layers/CompositeLayerSceneCapture.cpp
--- a/Engine/Plugins/Experimental/Compositing/Composite/Source/Composite/Private/Layers/CompositeLayerSceneCapture.cpp
+++ b/Engine/Plugins/Experimental/Compositing/Composite/Source/Composite/Private/Layers/CompositeLayerSceneCapture.cpp
@@ -55,7 +55,6 @@ void UCompositeLayerSceneCapture::PostEditChangeProperty(FPropertyChangedEvent&
 bool UCompositeLayerSceneCapture::GetProxy(FTraversalContext& InContext, FSceneRenderingBulkObjectAllocator& InFrameAllocator, FCompositeCorePassProxy*& OutProxy) const
 {
 	using namespace UE::CompositeCore;
-	constexpr int32 NumInputs = 2;
 
 	ACompositeActor* CompositeActor = GetTypedOuter<ACompositeActor>();
 	if (!IsValid(CompositeActor))
@@ -90,57 +89,34 @@ bool UCompositeLayerSceneCapture::GetProxy(FTraversalContext& InContext, FSceneR
 	ExternalTexture.Metadata.bDistorted = false;
 	ExternalTexture.Metadata.Encoding = UE::CompositeCore::EEncoding::Linear;
 
-	FPassInputDeclArray Inputs;
-	Inputs.SetNum(NumInputs);
-
 	const int32 TexIndex = InContext.ExternalTextures.Num() - 1;
-	Inputs[0].Set<FPassExternalResourceDesc>({ EXTERNAL_RANGE_START_ID + TexIndex });
-	Inputs[1] = GetDefaultSecondInput(InContext);
-
-	// Copy for subpasses before we move the declared inputs
-	FPassInputDecl BaseInput0 = Inputs[0];
+	FPassInputDecl PassInput;
+	PassInput.Set<FPassExternalResourceDesc>({ EXTERNAL_RANGE_START_ID + TexIndex });
 
-	FMergePassProxy* Proxy = InFrameAllocator.Create<FMergePassProxy>(MoveTemp(Inputs));
-	Proxy->MergeOp = GetOperation(InContext);
-
-	if (!Postprocessing.IsEmpty())
+	// Iterate backwards so that lower passes in the UI are executed first
+	for (int32 PassIndex = Postprocessing.Num() - 1; PassIndex >= 0; --PassIndex)
 	{
-		TArray<const FCompositeCorePassProxy*> SubPassProxies;
-		SubPassProxies.Reserve(Postprocessing.Num());
-
-		bool bFirstValidPass = true;
+		const TObjectPtr<UCompositePassBase>& SubPass = Postprocessing[PassIndex];
 
-		// Iterate backwards so that lower passes in the UI are executed first
-		for (int32 PassIndex = Postprocessing.Num() - 1; PassIndex >= 0; --PassIndex)
+		if (IsValid(SubPass) && SubPass->IsActive())
 		{
-			const TObjectPtr<UCompositePassBase>& SubPass = Postprocessing[PassIndex];
-
-			if (IsValid(SubPass) && SubPass->IsActive())
+			FCompositeCorePassProxy* ChildPassProxy;
+			if (SubPass->GetProxy(PassInput, InFrameAllocator, ChildPassProxy))
 			{
-				FPassInputDecl PassInput;
-				if (bFirstValidPass)
-				{
-					PassInput = BaseInput0;
-					bFirstValidPass = false;
-				}
-				else
-				{
-					// Default input from previous pass
-					PassInput.Set<FPassInternalResourceDesc>({});
-				}
-
-				FCompositeCorePassProxy* SubPassProxy;
-				if (SubPass->GetProxy(PassInput, InFrameAllocator, SubPassProxy))
-				{
-					SubPassProxies.Add(SubPassProxy);
-				}
+				// Next pass should receive the output of the current child pass.
+				PassInput.Set<const FCompositeCorePassProxy*>(ChildPassProxy);
 			}
 		}
-
-		// Apply sub-passes overriding the first input
-		Proxy->AddSubPasses(0, SubPassProxies);
 	}
 
+	FPassInputDeclArray Inputs;
+	Inputs.SetNum(FixedNumLayerInputs);
+	Inputs[0] = PassInput;
+	Inputs[1] = GetDefaultSecondInput(InContext);
+
+	FMergePassProxy* Proxy = InFrameAllocator.Create<FMergePassProxy>(MoveTemp(Inputs));
+	Proxy->MergeOp = GetOperation(InContext);
+
 	OutProxy = Proxy;
 	return true;
 }
diff --git a/Engine/Plugins/Experimental/Compositing/Composite/Source/Composite/Private/Layers/CompositeLayerShadowReflection.cpp b/Engine/Plugins/Experimental/Compositing/Composite/Source/Composite/Private/Layers/CompositeLayerShadowReflection.cpp
--- a/Engine/Plugins/Experimental/Compositing/Composite/Source/Composite/Private/Layers/CompositeLayerShadowReflection.cpp
+++ b/Engine/Plugins/Experimental/Compositing/Composite/Source/Composite/Private/Layers/CompositeLayerShadowReflection.cpp
@@ -43,7 +43,6 @@ void UCompositeLayerShadowReflection::PostEditChangeProperty(FPropertyChangedEve
 bool UCompositeLayerShadowReflection::GetProxy(FTraversalContext& InContext, FSceneRenderingBulkObjectAllocator& InFrameAllocator, FCompositeCorePassProxy*& OutProxy) const
 {
 	using namespace UE::CompositeCore;
-	constexpr int32 NumInputs = 2;
 
 	ACompositeActor* CompositeActor = GetTypedOuter<ACompositeActor>();
 	if (!IsValid(CompositeActor))
@@ -71,9 +70,9 @@ bool UCompositeLayerShadowReflection::GetProxy(FTraversalContext& InContext, FSc
 		ShadowReflectionPair[1]->HiddenActors = Actors;
 
 		FPassInputDeclArray Inputs;
-		Inputs.SetNum(NumInputs);
+		Inputs.SetNum(2);
 
-		for (int32 Index = 0; Index < NumInputs; ++Index)
+		for (int32 Index = 0; Index < 2; ++Index)
 		{
 			const TObjectPtr<UCompositeShadowReflectionCatcherComponent>& Component = ShadowReflectionPair[Index];
 			const FIntPoint& RenderResolution = CompositeActor->RenderResolution;
@@ -92,22 +91,19 @@ bool UCompositeLayerShadowReflection::GetProxy(FTraversalContext& InContext, FSc
 		DivisionProxy->MergeOp = ECompositeCoreMergeOp::Divide;
 		DivisionProxy->bUseLensDistortion = false;
 	}
-	else
+	
+	if (DivisionProxy == nullptr)
 	{
 		return false;
 	}
 
 	FPassInputDeclArray Inputs;
-	Inputs.SetNum(NumInputs);
-	// Input 0 to be replaced with shadow reflection
+	Inputs.SetNum(FixedNumLayerInputs);
+	Inputs[0].Set<const FCompositeCorePassProxy*>(DivisionProxy);
 	Inputs[1] = GetDefaultSecondInput(InContext);
 
 	Proxy = InFrameAllocator.Create<FMergePassProxy>(MoveTemp(Inputs));
 	Proxy->MergeOp = GetOperation(InContext);
-	if (DivisionProxy)
-	{
-		Proxy->AddSubPasses(0, { DivisionProxy });
-	}
 
 	OutProxy = Proxy;
 	return true;
diff --git a/Engine/Plugins/Experimental/Compositing/Composite/Source/Composite/Public/Layers/CompositeLayerBase.h b/Engine/Plugins/Experimental/Compositing/Composite/Source/Composite/Public/Layers/CompositeLayerBase.h
--- a/Engine/Plugins/Experimental/Compositing/Composite/Source/Composite/Public/Layers/CompositeLayerBase.h
+++ b/Engine/Plugins/Experimental/Compositing/Composite/Source/Composite/Public/Layers/CompositeLayerBase.h
@@ -72,6 +72,9 @@ public:
 	ECompositeCoreMergeOp Operation = ECompositeCoreMergeOp::Over;
 
 protected:
+	/** Fixed number of inputs used by layer merge passes. */
+	static constexpr int32 FixedNumLayerInputs = 2;
+
 	/** Convenience function that returns the merge operation depending on the traversal context. */
 	ECompositeCoreMergeOp GetOperation(const FTraversalContext& InContext) const;
 
diff --git a/Engine/Plugins/Experimental/Compositing/CompositeCore/Source/CompositeCore/Private/CompositeCoreSceneViewExtension.cpp b/Engine/Plugins/Experimental/Compositing/CompositeCore/Source/CompositeCore/Private/CompositeCoreSceneViewExtension.cpp
--- a/Engine/Plugins/Experimental/Compositing/CompositeCore/Source/CompositeCore/Private/CompositeCoreSceneViewExtension.cpp
+++ b/Engine/Plugins/Experimental/Compositing/CompositeCore/Source/CompositeCore/Private/CompositeCoreSceneViewExtension.cpp
@@ -565,59 +565,56 @@ bool FCompositeCoreSceneViewExtension::ApplyPasses_Recursive(
 		// Iterate over declared pass inputs
 		for (int32 InputIndex = 0; InputIndex < Pass->GetNumDeclaredInputs(); ++InputIndex)
 		{
-			const TArray<const FCompositeCorePassProxy*>* SubPasses = Pass->GetSubPasses(InputIndex);
-			if (SubPasses)
+			const FPassInputDecl& DeclaredInput = Pass->GetDeclaredInput(InputIndex);
+
+			// Input is the output result of a child (sub)pass
+			if (DeclaredInput.IsType<const FCompositeCorePassProxy*>())
 			{
-				// Recursively apply sub-passes per input index, automatically overriding inputs.
-				UE::CompositeCore::FPassOutput SubPassOutput;
-				if (ApplyPasses_Recursive(GraphBuilder, InView, PassInputs, OriginalInputs, PassContext, *SubPasses, ExternalTextures, RecursionLevel + 1, SubPassOutput))
+				const FCompositeCorePassProxy* ChildPassProxy = DeclaredInput.Get<const FCompositeCorePassProxy*>();
+				UE::CompositeCore::FPassOutput ChildPassOutput;
+
+				if (ApplyPasses_Recursive(GraphBuilder, InView, PassInputs, OriginalInputs, PassContext, { ChildPassProxy }, ExternalTextures, RecursionLevel + 1, ChildPassOutput))
 				{
-					UpdateNextPassInputs(GraphBuilder, SubPassOutput, InputIndex, PassInputs, ExternalTextures);
+					UpdateNextPassInputs(GraphBuilder, ChildPassOutput, InputIndex, PassInputs, ExternalTextures);
 				}
 			}
-			else
+			// If an external texture resource is expected, connect the current input index with the resolved identifier
+			else if (DeclaredInput.IsType<FPassExternalResourceDesc>())
 			{
-				// No sub-pass input override, fetch the regular input index
-				const FPassInputDecl& DeclaredInput = Pass->GetDeclaredInput(InputIndex);
+				const FPassExternalResourceDesc& Desc = DeclaredInput.Get<FPassExternalResourceDesc>();
+				const ResourceId DeclaredExternalId = Desc.Id;
 
-				// If an internal texture resource is expected, connect the current input index with either the original or current bass pass input.
-				if (DeclaredInput.IsType<FPassInternalResourceDesc>())
+				const FPassInput* SliceInput = ExternalTextures.Find(DeclaredExternalId);
+				if (ensureMsgf(SliceInput, TEXT("Invalid external input: %d"), DeclaredExternalId))
 				{
-					const FPassInternalResourceDesc& Desc = DeclaredInput.Get<FPassInternalResourceDesc>();
-					const int32 DeclaredInputIndex = Desc.Index;
-
-					if (Desc.bOriginalCopyBeforePasses)
-					{
-						if (ensureMsgf(OriginalInputs.IsValidIndex(DeclaredInputIndex), TEXT("Invalid internal input: %d"), DeclaredInputIndex))
-						{
-							PassInputs[InputIndex] = OriginalInputs[DeclaredInputIndex];
-						}
-					}
-					else
-					{
-						if (ensureMsgf(BasePassInputs.IsValidIndex(DeclaredInputIndex), TEXT("Invalid internal input: %d"), DeclaredInputIndex))
-						{
-							PassInputs[InputIndex] = BasePassInputs[DeclaredInputIndex];
-						}
-					}
+					PassInputs[InputIndex] = *SliceInput;
 				}
-				// If an external texture resource is expected, connect the current input index with the resolved identifier
-				else if (DeclaredInput.IsType<FPassExternalResourceDesc>())
-				{
-					const FPassExternalResourceDesc& Desc = DeclaredInput.Get<FPassExternalResourceDesc>();
-					const ResourceId DeclaredExternalId = Desc.Id;
+			}
+			// If an internal texture resource is expected, connect the current input index with either the original or current bass pass input.
+			else if (DeclaredInput.IsType<FPassInternalResourceDesc>())
+			{
+				const FPassInternalResourceDesc& Desc = DeclaredInput.Get<FPassInternalResourceDesc>();
+				const int32 DeclaredInputIndex = Desc.Index;
 
-					const FPassInput* SliceInput = ExternalTextures.Find(DeclaredExternalId);
-					if (ensureMsgf(SliceInput, TEXT("Invalid external input: %d"), DeclaredExternalId))
+				if (Desc.bOriginalCopyBeforePasses)
+				{
+					if (ensureMsgf(OriginalInputs.IsValidIndex(DeclaredInputIndex), TEXT("Invalid internal input: %d"), DeclaredInputIndex))
 					{
-						PassInputs[InputIndex] = *SliceInput;
+						PassInputs[InputIndex] = OriginalInputs[DeclaredInputIndex];
 					}
 				}
 				else
 				{
-					checkNoEntry();
+					if (ensureMsgf(BasePassInputs.IsValidIndex(DeclaredInputIndex), TEXT("Invalid internal input: %d"), DeclaredInputIndex))
+					{
+						PassInputs[InputIndex] = BasePassInputs[DeclaredInputIndex];
+					}
 				}
 			}
+			else
+			{
+				checkNoEntry();
+			}
 		}
 
 		if (bIsLastPass)
diff --git a/Engine/Plugins/Experimental/Compositing/CompositeCore/Source/CompositeCore/Private/Passes/CompositeCorePassProxy.cpp b/Engine/Plugins/Experimental/Compositing/CompositeCore/Source/CompositeCore/Private/Passes/CompositeCorePassProxy.cpp
--- a/Engine/Plugins/Experimental/Compositing/CompositeCore/Source/CompositeCore/Private/Passes/CompositeCorePassProxy.cpp
+++ b/Engine/Plugins/Experimental/Compositing/CompositeCore/Source/CompositeCore/Private/Passes/CompositeCorePassProxy.cpp
@@ -99,15 +99,6 @@ FCompositeCorePassProxy::FCompositeCorePassProxy(UE::CompositeCore::FPassInputDe
 	, PassOutputOverride(MoveTemp(InPassOutputOverride))
 { }
 
-
-void FCompositeCorePassProxy::AddSubPasses(int32 InputIndex, TArray<const FCompositeCorePassProxy*> InSubPasses)
-{
-	if (!InSubPasses.IsEmpty())
-	{
-		SubPasses.Emplace(InputIndex, MoveTemp(InSubPasses));
-	}
-}
-
 FScreenPassRenderTarget FCompositeCorePassProxy::CreateOutputRenderTarget(FRDGBuilder& GraphBuilder,
 	const FSceneView& InView, const FIntRect& OutputViewRect, FRDGTextureDesc OutputDesc, const TCHAR* InName)
 {
diff --git a/Engine/Plugins/Experimental/Compositing/CompositeCore/Source/CompositeCore/Public/Passes/CompositeCorePassProxy.h b/Engine/Plugins/Experimental/Compositing/CompositeCore/Source/CompositeCore/Public/Passes/CompositeCorePassProxy.h
--- a/Engine/Plugins/Experimental/Compositing/CompositeCore/Source/CompositeCore/Public/Passes/CompositeCorePassProxy.h
+++ b/Engine/Plugins/Experimental/Compositing/CompositeCore/Source/CompositeCore/Public/Passes/CompositeCorePassProxy.h
@@ -80,8 +80,8 @@ namespace UE
 			ResourceId Id = 0;
 		};
 
-		/** Pass input declaration, referring to either internally or externally managed textures. */
-		using FPassInputDecl = TVariant<FPassInternalResourceDesc, FPassExternalResourceDesc>;
+		/** Pass input declaration, referring to internal textures, external textures or the output of another pass proxy. */
+		using FPassInputDecl = TVariant<FPassInternalResourceDesc, FPassExternalResourceDesc, const FCompositeCorePassProxy*>;
 
 		/** Enum to facilitate the declaration of passes that have one default internal input. */
 		enum EDefaultPassInputDecl { DefaultPassInputDecl };
@@ -303,15 +303,6 @@ public:
 		PassOutputOverride.Reset();
 	}
 
-	/** Return sub-passes per input index. */
-	const TArray<const FCompositeCorePassProxy*>* GetSubPasses(int32 InputIndex) const
-	{
-		return SubPasses.Find(InputIndex);
-	}
-
-	/** Add sub-passes per input index. */
-	UE_API void AddSubPasses(int32 InputIndex, TArray<const FCompositeCorePassProxy*> InSubPasses);
-
 	/** Convenience function to create an output render target with the specified resolution. */
 	static UE_API FScreenPassRenderTarget CreateOutputRenderTarget(FRDGBuilder& GraphBuilder,
 		const FSceneView& InView, const FIntRect& OutputViewRect, FRDGTextureDesc OutputDesc, const TCHAR* InName);
@@ -326,9 +317,6 @@ protected:
 
 	/** Optional pass output override. */
 	TOptional<UE::CompositeCore::ResourceId> PassOutputOverride;
-
-	/** Map of sub-passes per input index. */
-	TSortedMap<int32, TArray<const FCompositeCorePassProxy*>> SubPasses;
 };
 
 #undef UE_API
