Composite: Merge pass proxy constructor cleanup.

* Cleanup parent FCompositeCorePassProxy constructors.
* Added lens distortion handling enum since optional boolean must be avoided as a last argument.

diff --git a/Engine/Plugins/Experimental/Compositing/Composite/Source/Composite/Private/Layers/CompositeLayerBase.cpp b/Engine/Plugins/Experimental/Compositing/Composite/Source/Composite/Private/Layers/CompositeLayerBase.cpp
--- a/Engine/Plugins/Experimental/Compositing/Composite/Source/Composite/Private/Layers/CompositeLayerBase.cpp
+++ b/Engine/Plugins/Experimental/Compositing/Composite/Source/Composite/Private/Layers/CompositeLayerBase.cpp
@@ -69,7 +69,7 @@ const FIntPoint UCompositeLayerBase::GetRenderResolution() const
 	return FCompositeRenderTargetPool::DefaultSize;
 }
 
-ECompositeCoreMergeOp UCompositeLayerBase::GetOperation(const FTraversalContext& InContext) const
+ECompositeCoreMergeOp UCompositeLayerBase::GetMergeOperation(const FTraversalContext& InContext) const
 {
 	return InContext.bIsSolo ? ECompositeCoreMergeOp::None : Operation;
 }
diff --git a/Engine/Plugins/Experimental/Compositing/Composite/Source/Composite/Private/Layers/CompositeLayerMainRender.cpp b/Engine/Plugins/Experimental/Compositing/Composite/Source/Composite/Private/Layers/CompositeLayerMainRender.cpp
--- a/Engine/Plugins/Experimental/Compositing/Composite/Source/Composite/Private/Layers/CompositeLayerMainRender.cpp
+++ b/Engine/Plugins/Experimental/Compositing/Composite/Source/Composite/Private/Layers/CompositeLayerMainRender.cpp
@@ -25,10 +25,7 @@ bool UCompositeLayerMainRender::GetProxy(FTraversalContext& InContext, FSceneRen
 	Inputs[0].Set<FPassInternalResourceDesc>(Desc0);
 	Inputs[1] = GetDefaultSecondInput(InContext);
 
-	FMergePassProxy* Proxy = InFrameAllocator.Create<FMergePassProxy>(MoveTemp(Inputs));
-	Proxy->MergeOp = GetOperation(InContext);
-
-	OutProxy = Proxy;
+	OutProxy = InFrameAllocator.Create<FMergePassProxy>(MoveTemp(Inputs), GetMergeOperation(InContext));
 	return true;
 }
 
diff --git a/Engine/Plugins/Experimental/Compositing/Composite/Source/Composite/Private/Layers/CompositeLayerPlate.cpp b/Engine/Plugins/Experimental/Compositing/Composite/Source/Composite/Private/Layers/CompositeLayerPlate.cpp
--- a/Engine/Plugins/Experimental/Compositing/Composite/Source/Composite/Private/Layers/CompositeLayerPlate.cpp
+++ b/Engine/Plugins/Experimental/Compositing/Composite/Source/Composite/Private/Layers/CompositeLayerPlate.cpp
@@ -509,11 +509,7 @@ bool UCompositeLayerPlate::GetProxy(FTraversalContext& InContext, FSceneRenderin
 	Inputs[0] = PassInput;
 	Inputs[1] = GetDefaultSecondInput(InContext);
 
-	FMergePassProxy* Proxy = InFrameAllocator.Create<FMergePassProxy>(MoveTemp(Inputs));
-	Proxy->MergeOp = GetOperation(InContext);
-	Proxy->bUseLensDistortion = false;
-
-	OutProxy = Proxy;
+	OutProxy = InFrameAllocator.Create<FMergePassProxy>(MoveTemp(Inputs), GetMergeOperation(InContext), ELensDistortionHandling::Disabled);
 	return true;
 }
 
diff --git a/Engine/Plugins/Experimental/Compositing/Composite/Source/Composite/Private/Layers/CompositeLayerSceneCapture.cpp b/Engine/Plugins/Experimental/Compositing/Composite/Source/Composite/Private/Layers/CompositeLayerSceneCapture.cpp
--- a/Engine/Plugins/Experimental/Compositing/Composite/Source/Composite/Private/Layers/CompositeLayerSceneCapture.cpp
+++ b/Engine/Plugins/Experimental/Compositing/Composite/Source/Composite/Private/Layers/CompositeLayerSceneCapture.cpp
@@ -112,10 +112,7 @@ bool UCompositeLayerSceneCapture::GetProxy(FTraversalContext& InContext, FSceneR
 	Inputs[0] = PassInput;
 	Inputs[1] = GetDefaultSecondInput(InContext);
 
-	FMergePassProxy* Proxy = InFrameAllocator.Create<FMergePassProxy>(MoveTemp(Inputs));
-	Proxy->MergeOp = GetOperation(InContext);
-
-	OutProxy = Proxy;
+	OutProxy = InFrameAllocator.Create<FMergePassProxy>(MoveTemp(Inputs), GetMergeOperation(InContext));
 	return true;
 }
 
diff --git a/Engine/Plugins/Experimental/Compositing/Composite/Source/Composite/Private/Layers/CompositeLayerShadowReflection.cpp b/Engine/Plugins/Experimental/Compositing/Composite/Source/Composite/Private/Layers/CompositeLayerShadowReflection.cpp
--- a/Engine/Plugins/Experimental/Compositing/Composite/Source/Composite/Private/Layers/CompositeLayerShadowReflection.cpp
+++ b/Engine/Plugins/Experimental/Compositing/Composite/Source/Composite/Private/Layers/CompositeLayerShadowReflection.cpp
@@ -88,9 +88,7 @@ bool UCompositeLayerShadowReflection::GetProxy(FTraversalContext& InContext, FSc
 			Inputs[Index].Set<FPassExternalResourceDesc>({ EXTERNAL_RANGE_START_ID + TexIndex });
 		}
 
-		DivisionProxy = InFrameAllocator.Create<FMergePassProxy>(MoveTemp(Inputs));
-		DivisionProxy->MergeOp = ECompositeCoreMergeOp::Divide;
-		DivisionProxy->bUseLensDistortion = false;
+		DivisionProxy = InFrameAllocator.Create<FMergePassProxy>(MoveTemp(Inputs), ECompositeCoreMergeOp::Divide, ELensDistortionHandling::Disabled);
 	}
 	
 	if (DivisionProxy == nullptr)
@@ -103,10 +101,7 @@ bool UCompositeLayerShadowReflection::GetProxy(FTraversalContext& InContext, FSc
 	Inputs[0].Set<const FCompositeCorePassProxy*>(DivisionProxy);
 	Inputs[1] = GetDefaultSecondInput(InContext);
 
-	Proxy = InFrameAllocator.Create<FMergePassProxy>(MoveTemp(Inputs));
-	Proxy->MergeOp = GetOperation(InContext);
-
-	OutProxy = Proxy;
+	OutProxy = InFrameAllocator.Create<FMergePassProxy>(MoveTemp(Inputs), GetMergeOperation(InContext));
 	return true;
 }
 
diff --git a/Engine/Plugins/Experimental/Compositing/Composite/Source/Composite/Public/Layers/CompositeLayerBase.h b/Engine/Plugins/Experimental/Compositing/Composite/Source/Composite/Public/Layers/CompositeLayerBase.h
--- a/Engine/Plugins/Experimental/Compositing/Composite/Source/Composite/Public/Layers/CompositeLayerBase.h
+++ b/Engine/Plugins/Experimental/Compositing/Composite/Source/Composite/Public/Layers/CompositeLayerBase.h
@@ -79,7 +79,7 @@ protected:
 	const FIntPoint GetRenderResolution() const;
 
 	/** Convenience function that returns the merge operation depending on the traversal context. */
-	ECompositeCoreMergeOp GetOperation(const FTraversalContext& InContext) const;
+	ECompositeCoreMergeOp GetMergeOperation(const FTraversalContext& InContext) const;
 
 	/** Convenience function that returns the default secondary input depending on the traversal context. */
 	UE::CompositeCore::FPassInputDecl GetDefaultSecondInput(const FTraversalContext& InContext) const;
diff --git a/Engine/Plugins/Experimental/Compositing/CompositeCore/Source/CompositeCore/Private/CompositeCoreSceneViewExtension.cpp b/Engine/Plugins/Experimental/Compositing/CompositeCore/Source/CompositeCore/Private/CompositeCoreSceneViewExtension.cpp
--- a/Engine/Plugins/Experimental/Compositing/CompositeCore/Source/CompositeCore/Private/CompositeCoreSceneViewExtension.cpp
+++ b/Engine/Plugins/Experimental/Compositing/CompositeCore/Source/CompositeCore/Private/CompositeCoreSceneViewExtension.cpp
@@ -359,7 +359,7 @@ void FCompositeCoreSceneViewExtension::PrePostProcessPass_RenderThread(FRDGBuild
 
 	if ((CompositeRenderPassPtr != nullptr) && CVarCompositeCoreApplyFXAA.GetValueOnRenderThread())
 	{
-		static FFXAAPassProxy FXAAPassProxy = FFXAAPassProxy(DefaultPassInputDecl);
+		static FFXAAPassProxy FXAAPassProxy = FFXAAPassProxy(GetDefaultInputDeclArray());
 
 		// Set the composite render target as the pass input
 		FPassInputArray PassInputs;
diff --git a/Engine/Plugins/Experimental/Compositing/CompositeCore/Source/CompositeCore/Private/Passes/CompositeCorePassMergeProxy.cpp b/Engine/Plugins/Experimental/Compositing/CompositeCore/Source/CompositeCore/Private/Passes/CompositeCorePassMergeProxy.cpp
--- a/Engine/Plugins/Experimental/Compositing/CompositeCore/Source/CompositeCore/Private/Passes/CompositeCorePassMergeProxy.cpp
+++ b/Engine/Plugins/Experimental/Compositing/CompositeCore/Source/CompositeCore/Private/Passes/CompositeCorePassMergeProxy.cpp
@@ -114,6 +114,13 @@ namespace UE
 			return Parameters;
 		}
 
+		FMergePassProxy::FMergePassProxy(FPassInputDeclArray InPassDeclaredInputs, ECompositeCoreMergeOp InMergeOp, ELensDistortionHandling bInUseLensDistortion)
+			: FCompositeCorePassProxy(MoveTemp(InPassDeclaredInputs))
+			, MergeOp(InMergeOp)
+			, bUseLensDistortion(bInUseLensDistortion)
+		{
+		}
+
 		FPassOutput FMergePassProxy::Add(FRDGBuilder& GraphBuilder, const FSceneView& InView, const FPassInputArray& Inputs, const FPassContext& PassContext) const
 		{
 			RDG_EVENT_SCOPE_STAT(GraphBuilder, FCompositeCoreMerge, "CompositeCore.Merge");
@@ -126,7 +133,7 @@ namespace UE
 			FResourceMetadata OutMetadata = {};
 
 			const FLensDistortionLUT& LensDistortionLUT = LensDistortion::GetLUTUnsafe(InView);
-			const bool bIsLensDistortionActive = bUseLensDistortion && LensDistortionLUT.IsEnabled()
+			const bool bIsLensDistortionActive = static_cast<bool>(bUseLensDistortion) && LensDistortionLUT.IsEnabled()
 				&& PassContext.Location >= ISceneViewExtension::EPostProcessingPass::SSRInput;
 
 			if (bIsLensDistortionActive && PassContext.bOutputSceneColor)
diff --git a/Engine/Plugins/Experimental/Compositing/CompositeCore/Source/CompositeCore/Private/Passes/CompositeCorePassProxy.cpp b/Engine/Plugins/Experimental/Compositing/CompositeCore/Source/CompositeCore/Private/Passes/CompositeCorePassProxy.cpp
--- a/Engine/Plugins/Experimental/Compositing/CompositeCore/Source/CompositeCore/Private/Passes/CompositeCorePassProxy.cpp
+++ b/Engine/Plugins/Experimental/Compositing/CompositeCore/Source/CompositeCore/Private/Passes/CompositeCorePassProxy.cpp
@@ -84,19 +84,17 @@ namespace UE
 
 			return bIsEmpty;
 		}
-	}
-}
-
 
-FCompositeCorePassProxy::FCompositeCorePassProxy(UE::CompositeCore::EDefaultPassInputDecl)
-	: PassDeclaredInputs({ UE::CompositeCore::FPassInputDecl(TInPlaceType<UE::CompositeCore::FPassInternalResourceDesc>(), UE::CompositeCore::FPassInternalResourceDesc{}) })
-	, PassOutputOverride()
-{
+		FPassInputDeclArray GetDefaultInputDeclArray()
+		{
+			return FPassInputDeclArray{ FPassInputDecl{TInPlaceType<FPassInternalResourceDesc>(), FPassInternalResourceDesc{}} };
+		}
+	}
 }
 
-FCompositeCorePassProxy::FCompositeCorePassProxy(UE::CompositeCore::FPassInputDeclArray InPassDeclaredInputs, TOptional<UE::CompositeCore::ResourceId> InPassOutputOverride)
+FCompositeCorePassProxy::FCompositeCorePassProxy(UE::CompositeCore::FPassInputDeclArray InPassDeclaredInputs)
 	: PassDeclaredInputs(MoveTemp(InPassDeclaredInputs))
-	, PassOutputOverride(MoveTemp(InPassOutputOverride))
+	, PassOutputOverride()
 { }
 
 FScreenPassRenderTarget FCompositeCorePassProxy::CreateOutputRenderTarget(FRDGBuilder& GraphBuilder,
diff --git a/Engine/Plugins/Experimental/Compositing/CompositeCore/Source/CompositeCore/Public/Passes/CompositeCorePassMergeProxy.h b/Engine/Plugins/Experimental/Compositing/CompositeCore/Source/CompositeCore/Public/Passes/CompositeCorePassMergeProxy.h
--- a/Engine/Plugins/Experimental/Compositing/CompositeCore/Source/CompositeCore/Public/Passes/CompositeCorePassMergeProxy.h
+++ b/Engine/Plugins/Experimental/Compositing/CompositeCore/Source/CompositeCore/Public/Passes/CompositeCorePassMergeProxy.h
@@ -36,18 +36,28 @@ namespace UE
 {
 	namespace CompositeCore
 	{
+		/** Describe whether lens distortion is handled in the merge shader pass. */
+		enum class ELensDistortionHandling : uint8
+		{
+			Disabled = 0u,
+			Enabled = 1u
+		};
+
 		class FMergePassProxy : public FCompositeCorePassProxy
 		{
 		public:
 			using FCompositeCorePassProxy::FCompositeCorePassProxy;
 
+			/** Constructor */
+			UE_API FMergePassProxy(FPassInputDeclArray InPassDeclaredInputs, ECompositeCoreMergeOp InMergeOp, ELensDistortionHandling bInUseLensDistortion = ELensDistortionHandling::Enabled);
+
 			UE_API FPassOutput Add(FRDGBuilder& GraphBuilder, const FSceneView& InView, const FPassInputArray& Inputs, const FPassContext& PassContext) const override;
 
 			/** Merge operation. */
 			ECompositeCoreMergeOp MergeOp = ECompositeCoreMergeOp::Over;
 
 			/** Flag to enable or disable application of engine-provided lens distortion LUTs, true by default. */
-			bool bUseLensDistortion = true;
+			ELensDistortionHandling bUseLensDistortion = ELensDistortionHandling::Enabled;
 		};
 	}
 }
diff --git a/Engine/Plugins/Experimental/Compositing/CompositeCore/Source/CompositeCore/Public/Passes/CompositeCorePassProxy.h b/Engine/Plugins/Experimental/Compositing/CompositeCore/Source/CompositeCore/Public/Passes/CompositeCorePassProxy.h
--- a/Engine/Plugins/Experimental/Compositing/CompositeCore/Source/CompositeCore/Public/Passes/CompositeCorePassProxy.h
+++ b/Engine/Plugins/Experimental/Compositing/CompositeCore/Source/CompositeCore/Public/Passes/CompositeCorePassProxy.h
@@ -82,13 +82,13 @@ namespace UE
 
 		/** Pass input declaration, referring to internal textures, external textures or the output of another pass proxy. */
 		using FPassInputDecl = TVariant<FPassInternalResourceDesc, FPassExternalResourceDesc, const FCompositeCorePassProxy*>;
-
-		/** Enum to facilitate the declaration of passes that have one default internal input. */
-		enum EDefaultPassInputDecl { DefaultPassInputDecl };
 		
 		/** Array of pass input declarations. */
 		using FPassInputDeclArray = TArray<FPassInputDecl>;
 
+		/** Returns a default input array declaring one internal texture resource. */
+		UE_API FPassInputDeclArray GetDefaultInputDeclArray();
+
 		/** External texture resource and its accompanying metadata. */
 		struct FExternalTexture
 		{
@@ -256,14 +256,11 @@ public:
 	/** Default constructor. */
 	FCompositeCorePassProxy() = default;
 
-	/** Constructor for a single default internal input. */
-	UE_API FCompositeCorePassProxy(UE::CompositeCore::EDefaultPassInputDecl);
-
 	/** Constructor. */
-	UE_API FCompositeCorePassProxy(UE::CompositeCore::FPassInputDeclArray InPassDeclaredInputs, TOptional<UE::CompositeCore::ResourceId> InPassOutputOverride = {});
+	UE_API FCompositeCorePassProxy(UE::CompositeCore::FPassInputDeclArray InPassDeclaredInputs);
 
 	/** Virtual destructor */
-	virtual ~FCompositeCorePassProxy() { };
+	UE_API virtual ~FCompositeCorePassProxy() = default;
 
 	/** Render-thread add pass method to override. */
 	virtual UE::CompositeCore::FPassOutput Add(
