Composite: Recyling pass work between SSR input & after tonemap locations in post-processing.

* Improved proxy generation logic on the composite actor for sharing work between locations, for all passes until scene textures are used (which must remain location specific).
* Removed FPassOutput type which was actually useless.
* Fixed minor note from previous review.

diff --git a/Engine/Plugins/Experimental/Compositing/Composite/Source/Composite/Private/CompositeActor.cpp b/Engine/Plugins/Experimental/Compositing/Composite/Source/Composite/Private/CompositeActor.cpp
--- a/Engine/Plugins/Experimental/Compositing/Composite/Source/Composite/Private/CompositeActor.cpp
+++ b/Engine/Plugins/Experimental/Compositing/Composite/Source/Composite/Private/CompositeActor.cpp
@@ -71,42 +71,49 @@ namespace UE
 				return SoloLayer;
 			}
 
-			/** Convenience function to register layers in the traversal context, while returning proxies for each layer input. */
-			void RegisterLayers(UCompositeLayerBase::FTraversalContext& Context, TArrayView<const TObjectPtr<UCompositeLayerBase>> InLayers, FSceneRenderingBulkObjectAllocator& InFrameAllocator, TArray<const FCompositeCorePassProxy*>& OutProxies)
+			/**
+			* Convenience function to register layers in the traversal context, while returning proxies for each layer input.
+			* 
+			* Returns an index indicating the first proxy whose output can no longer be shared between post-processing locations.
+			*/
+			int32 RegisterLayers(UCompositeLayerBase::FTraversalContext& Context, TArrayView<const TObjectPtr<UCompositeLayerBase>> InLayers, FSceneRenderingBulkObjectAllocator& InFrameAllocator, TArray<FCompositeCorePassProxy*>& OutProxies)
 			{
+				int32 LastLayerIndex = INDEX_NONE;
+				int32 BranchLayerIndex = INDEX_NONE;
 				OutProxies.Reserve(InLayers.Num());
-				Context.bIsSolo = (InLayers.Num() == 1);
 
-				// Iterate backwards so that lower passes in the UI are executed first
-				const int32 FirstPassIndex = InLayers.Num() - 1;
-
-				TOptional<int32> FirstValidPassIndex;
-				for (int32 LayerIndex = FirstPassIndex; LayerIndex >= 0; --LayerIndex)
-				{
-					const TObjectPtr<UCompositeLayerBase>& Layer = InLayers[LayerIndex];
-					if (IsValid(Layer) && Layer->bIsEnabled)
-					{
-						FirstValidPassIndex = LayerIndex;
-						break;
-					}
-				}
-
-				for (int32 LayerIndex = FirstPassIndex; LayerIndex >= 0; --LayerIndex)
+				for (const TObjectPtr<UCompositeLayerBase>& Layer : InLayers)
 				{
-					const TObjectPtr<UCompositeLayerBase>& Layer = InLayers[LayerIndex];
 					if (!IsValid(Layer) || !Layer->bIsEnabled)
 					{
 						continue;
 					}
 
-					Context.bIsFirstPass = (LayerIndex == FirstValidPassIndex);
-
 					FCompositeCorePassProxy* OutProxy;
 					if (Layer->GetProxy(Context, InFrameAllocator, OutProxy))
 					{
 						OutProxies.Add(OutProxy);
+						Context.bIsFirstPass = false;
+
+						const int32 ProxyIndex = OutProxies.Num() - 1;
+						// If the pass needed scene textures specific to each post-processing location, its output
+						// cannot be recyled in multiple locations. We therefore mark it as the "branch" layer.
+						if (Context.bNeedsSceneTextures)
+						{
+							if (BranchLayerIndex == INDEX_NONE)
+							{
+								BranchLayerIndex = ProxyIndex;
+							}
+						}
+
+						// Since merge passes handle color encodings & lens distortion, and
+						// since SSR Input & tonemap differ in both regards, we must at least
+						// run one merge pass for each location currently.
+						LastLayerIndex = ProxyIndex;
 					}
 				}
+
+				return (BranchLayerIndex != INDEX_NONE) ? BranchLayerIndex : LastLayerIndex;
 			}
 
 			/* Gather render work (per frame) given the current composite actor properties. */
@@ -128,22 +135,43 @@ namespace UE
 				else
 				{
 					Layers = InCompositeActor.CompositeLayers;
+					// TODO: Remove once we have UI customization displaying layers in the correct bottom-up order
+					Algo::Reverse(Layers);
 				}
 
-				{
-					UCompositeLayerBase::FTraversalContext Context;
-					RegisterLayers(Context, Layers, *RenderWork.FrameAllocator, RenderWork.FramePasses.Add(ISceneViewExtension::EPostProcessingPass::Tonemap));
-					
-					RenderWork.ExternalInputs = Context.ExternalTextures;
-					RenderWork.PreprocessingPasses = Context.PreprocessingPasses;
-				}
+				TArray<FCompositeCorePassProxy*> LayerProxies;
+				UCompositeLayerBase::FTraversalContext Context;
+				Context.bIsSolo = (Layers.Num() == 1);
+
+				const int32 BranchLayerIndex = RegisterLayers(Context, Layers, *RenderWork.FrameAllocator, LayerProxies);
+				const int32 BranchSourceIndex = BranchLayerIndex - 1;
 
-				// TODO: Recycle work between SSR input & after-tonemap composite step.
 				if (InCompositeActor.bEnableScreenSpaceReflections)
 				{
-					UCompositeLayerBase::FTraversalContext TempContext;
-					RegisterLayers(TempContext, Layers, *RenderWork.FrameAllocator, RenderWork.FramePasses.Add(ISceneViewExtension::EPostProcessingPass::SSRInput));
+					TArray<FCompositeCorePassProxy*> BranchedLayerProxies;
+					Context.bIsFirstPass = true;
+					RegisterLayers(Context, Layers, *RenderWork.FrameAllocator, BranchedLayerProxies);
+
+					// Trim all proxies before our branch index (for the tonemap location) to instead recycle previous work.
+					BranchedLayerProxies.RemoveAt(0, BranchLayerIndex);
+
+					if (ensureMsgf(!BranchedLayerProxies.IsEmpty(), TEXT("Unexpected empty branch proxy array.")) && LayerProxies.IsValidIndex(BranchSourceIndex))
+					{
+						// This is where we connect to the previous SSR input work branch.
+						// Input0 will be merged on top of Input1, the previous proxy's output.
+						BranchedLayerProxies[0]->GetDeclaredInput(1).Set<const FCompositeCorePassProxy*>(LayerProxies[BranchSourceIndex]);
+					}
+
+					RenderWork.FramePasses.Emplace(ISceneViewExtension::EPostProcessingPass::SSRInput, LayerProxies);
+					RenderWork.FramePasses.Emplace(ISceneViewExtension::EPostProcessingPass::Tonemap, BranchedLayerProxies);
 				}
+				else
+				{
+					RenderWork.FramePasses.Emplace(ISceneViewExtension::EPostProcessingPass::Tonemap, LayerProxies);
+				}
+
+				RenderWork.ExternalInputs = Context.ExternalTextures;
+				RenderWork.PreprocessingPasses = Context.PreprocessingPasses;
 
 				return RenderWork;
 			}
diff --git a/Engine/Plugins/Experimental/Compositing/Composite/Source/Composite/Private/Layers/CompositeLayerBase.cpp b/Engine/Plugins/Experimental/Compositing/Composite/Source/Composite/Private/Layers/CompositeLayerBase.cpp
--- a/Engine/Plugins/Experimental/Compositing/Composite/Source/Composite/Private/Layers/CompositeLayerBase.cpp
+++ b/Engine/Plugins/Experimental/Compositing/Composite/Source/Composite/Private/Layers/CompositeLayerBase.cpp
@@ -57,6 +57,31 @@ void UCompositeLayerBase::PostEditChangeProperty(FPropertyChangedEvent& Property
 }
 #endif //WITH_EDITOR
 
+void UCompositeLayerBase::AddChildPasses(UE::CompositeCore::FPassInputDecl& InBasePassInput, FTraversalContext& InContext, FSceneRenderingBulkObjectAllocator& InFrameAllocator, TArrayView<const TObjectPtr<UCompositePassBase>> InPasses) const
+{
+	// Iterate backwards so that lower passes in the UI are executed first
+	// TODO: Flip back once we have UI customization displaying passes in the correct bottom-up order
+	for (int32 PassIndex = InPasses.Num() - 1; PassIndex >= 0; --PassIndex)
+	{
+		const TObjectPtr<UCompositePassBase>& SubPass = InPasses[PassIndex];
+
+		if (IsValid(SubPass) && SubPass->IsActive())
+		{
+			FCompositeCorePassProxy* ChildPassProxy;
+			if (SubPass->GetProxy(InBasePassInput, InFrameAllocator, ChildPassProxy))
+			{
+				// Next pass should receive the output of the current child pass.
+				InBasePassInput.Set<const FCompositeCorePassProxy*>(ChildPassProxy);
+
+				if (SubPass->NeedsSceneTextures())
+				{
+					InContext.bNeedsSceneTextures = true;
+				}
+			}
+		}
+	}
+}
+
 const FIntPoint UCompositeLayerBase::GetRenderResolution() const
 {
 	const ACompositeActor* CompositeActor = GetTypedOuter<ACompositeActor>();
diff --git a/Engine/Plugins/Experimental/Compositing/Composite/Source/Composite/Private/Layers/CompositeLayerMainRender.cpp b/Engine/Plugins/Experimental/Compositing/Composite/Source/Composite/Private/Layers/CompositeLayerMainRender.cpp
--- a/Engine/Plugins/Experimental/Compositing/Composite/Source/Composite/Private/Layers/CompositeLayerMainRender.cpp
+++ b/Engine/Plugins/Experimental/Compositing/Composite/Source/Composite/Private/Layers/CompositeLayerMainRender.cpp
@@ -25,7 +25,8 @@ bool UCompositeLayerMainRender::GetProxy(FTraversalContext& InContext, FSceneRen
 	Inputs[0].Set<FPassInternalResourceDesc>(Desc0);
 	Inputs[1] = GetDefaultSecondInput(InContext);
 
-	OutProxy = InFrameAllocator.Create<FMergePassProxy>(MoveTemp(Inputs), GetMergeOperation(InContext));
+	InContext.bNeedsSceneTextures = true;
+	OutProxy = InFrameAllocator.Create<FMergePassProxy>(MoveTemp(Inputs), GetMergeOperation(InContext), TEXT("MainRender"));
 	return true;
 }
 
diff --git a/Engine/Plugins/Experimental/Compositing/Composite/Source/Composite/Private/Layers/CompositeLayerPlate.cpp b/Engine/Plugins/Experimental/Compositing/Composite/Source/Composite/Private/Layers/CompositeLayerPlate.cpp
--- a/Engine/Plugins/Experimental/Compositing/Composite/Source/Composite/Private/Layers/CompositeLayerPlate.cpp
+++ b/Engine/Plugins/Experimental/Compositing/Composite/Source/Composite/Private/Layers/CompositeLayerPlate.cpp
@@ -488,28 +488,14 @@ bool UCompositeLayerPlate::GetProxy(FTraversalContext& InContext, FSceneRenderin
 	FPassInputDecl PassInput;
 	PassInput.Set<FPassExternalResourceDesc>({ InputId });
 
-	// Iterate backwards so that lower passes in the UI are executed first
-	for (int32 PassIndex = Postprocessing.Num() - 1; PassIndex >= 0; --PassIndex)
-	{
-		const TObjectPtr<UCompositePassBase>& SubPass = Postprocessing[PassIndex];
-
-		if (IsValid(SubPass) && SubPass->IsActive())
-		{
-			FCompositeCorePassProxy* ChildPassProxy;
-			if (SubPass->GetProxy(PassInput, InFrameAllocator, ChildPassProxy))
-			{
-				// Next pass should receive the output of the current child pass.
-				PassInput.Set<const FCompositeCorePassProxy*>(ChildPassProxy);
-			}
-		}
-	}
+	AddChildPasses(PassInput, InContext, InFrameAllocator, Postprocessing);
 
 	FPassInputDeclArray Inputs;
 	Inputs.SetNum(FixedNumLayerInputs);
 	Inputs[0] = PassInput;
 	Inputs[1] = GetDefaultSecondInput(InContext);
 
-	OutProxy = InFrameAllocator.Create<FMergePassProxy>(MoveTemp(Inputs), GetMergeOperation(InContext), ELensDistortionHandling::Disabled);
+	OutProxy = InFrameAllocator.Create<FMergePassProxy>(MoveTemp(Inputs), GetMergeOperation(InContext), TEXT("Plate"), ELensDistortionHandling::Disabled);
 	return true;
 }
 
diff --git a/Engine/Plugins/Experimental/Compositing/Composite/Source/Composite/Private/Layers/CompositeLayerSceneCapture.cpp b/Engine/Plugins/Experimental/Compositing/Composite/Source/Composite/Private/Layers/CompositeLayerSceneCapture.cpp
--- a/Engine/Plugins/Experimental/Compositing/Composite/Source/Composite/Private/Layers/CompositeLayerSceneCapture.cpp
+++ b/Engine/Plugins/Experimental/Compositing/Composite/Source/Composite/Private/Layers/CompositeLayerSceneCapture.cpp
@@ -91,28 +91,14 @@ bool UCompositeLayerSceneCapture::GetProxy(FTraversalContext& InContext, FSceneR
 	FPassInputDecl PassInput;
 	PassInput.Set<FPassExternalResourceDesc>({ EXTERNAL_RANGE_START_ID + TexIndex });
 
-	// Iterate backwards so that lower passes in the UI are executed first
-	for (int32 PassIndex = Postprocessing.Num() - 1; PassIndex >= 0; --PassIndex)
-	{
-		const TObjectPtr<UCompositePassBase>& SubPass = Postprocessing[PassIndex];
-
-		if (IsValid(SubPass) && SubPass->IsActive())
-		{
-			FCompositeCorePassProxy* ChildPassProxy;
-			if (SubPass->GetProxy(PassInput, InFrameAllocator, ChildPassProxy))
-			{
-				// Next pass should receive the output of the current child pass.
-				PassInput.Set<const FCompositeCorePassProxy*>(ChildPassProxy);
-			}
-		}
-	}
+	AddChildPasses(PassInput, InContext, InFrameAllocator, Postprocessing);
 
 	FPassInputDeclArray Inputs;
 	Inputs.SetNum(FixedNumLayerInputs);
 	Inputs[0] = PassInput;
 	Inputs[1] = GetDefaultSecondInput(InContext);
 
-	OutProxy = InFrameAllocator.Create<FMergePassProxy>(MoveTemp(Inputs), GetMergeOperation(InContext));
+	OutProxy = InFrameAllocator.Create<FMergePassProxy>(MoveTemp(Inputs), GetMergeOperation(InContext), TEXT("SceneCapture"));
 	return true;
 }
 
diff --git a/Engine/Plugins/Experimental/Compositing/Composite/Source/Composite/Private/Layers/CompositeLayerShadowReflection.cpp b/Engine/Plugins/Experimental/Compositing/Composite/Source/Composite/Private/Layers/CompositeLayerShadowReflection.cpp
--- a/Engine/Plugins/Experimental/Compositing/Composite/Source/Composite/Private/Layers/CompositeLayerShadowReflection.cpp
+++ b/Engine/Plugins/Experimental/Compositing/Composite/Source/Composite/Private/Layers/CompositeLayerShadowReflection.cpp
@@ -88,7 +88,7 @@ bool UCompositeLayerShadowReflection::GetProxy(FTraversalContext& InContext, FSc
 			Inputs[Index].Set<FPassExternalResourceDesc>({ EXTERNAL_RANGE_START_ID + TexIndex });
 		}
 
-		DivisionProxy = InFrameAllocator.Create<FMergePassProxy>(MoveTemp(Inputs), ECompositeCoreMergeOp::Divide, ELensDistortionHandling::Disabled);
+		DivisionProxy = InFrameAllocator.Create<FMergePassProxy>(MoveTemp(Inputs), ECompositeCoreMergeOp::Divide, TEXT("ShadowRefl Div"), ELensDistortionHandling::Disabled);
 	}
 	
 	if (DivisionProxy == nullptr)
@@ -101,7 +101,7 @@ bool UCompositeLayerShadowReflection::GetProxy(FTraversalContext& InContext, FSc
 	Inputs[0].Set<const FCompositeCorePassProxy*>(DivisionProxy);
 	Inputs[1] = GetDefaultSecondInput(InContext);
 
-	OutProxy = InFrameAllocator.Create<FMergePassProxy>(MoveTemp(Inputs), GetMergeOperation(InContext));
+	OutProxy = InFrameAllocator.Create<FMergePassProxy>(MoveTemp(Inputs), GetMergeOperation(InContext), TEXT("ShadowRefl Mul"));
 	return true;
 }
 
diff --git a/Engine/Plugins/Experimental/Compositing/Composite/Source/Composite/Private/Passes/CompositePassCenteredScale.cpp b/Engine/Plugins/Experimental/Compositing/Composite/Source/Composite/Private/Passes/CompositePassCenteredScale.cpp
--- a/Engine/Plugins/Experimental/Compositing/Composite/Source/Composite/Private/Passes/CompositePassCenteredScale.cpp
+++ b/Engine/Plugins/Experimental/Compositing/Composite/Source/Composite/Private/Passes/CompositePassCenteredScale.cpp
@@ -35,7 +35,7 @@ namespace UE
 		public:
 			using FCompositeCorePassProxy::FCompositeCorePassProxy;
 
-			FPassOutput Add(FRDGBuilder& GraphBuilder, const FSceneView& InView, const FPassInputArray& Inputs, const FPassContext& PassContext) const override
+			FPassTexture Add(FRDGBuilder& GraphBuilder, const FSceneView& InView, const FPassInputArray& Inputs, const FPassContext& PassContext) const override
 			{
 				RDG_EVENT_SCOPE_STAT(GraphBuilder, FCompositeCenteredScale, "Composite.CenteredScale");
 				RDG_GPU_STAT_SCOPE(GraphBuilder, FCompositeCenteredScale);
@@ -72,7 +72,7 @@ namespace UE
 					Parameters
 				);
 
-				return UE::CompositeCore::FPassOutput{ MoveTemp(Output), Metadata, PassOutputOverride };
+				return UE::CompositeCore::FPassTexture{ MoveTemp(Output), Metadata };
 			}
 
 			FVector2f ScaleUV = FVector2f::One();
diff --git a/Engine/Plugins/Experimental/Compositing/Composite/Source/Composite/Private/Passes/CompositePassColorGrade.cpp b/Engine/Plugins/Experimental/Compositing/Composite/Source/Composite/Private/Passes/CompositePassColorGrade.cpp
--- a/Engine/Plugins/Experimental/Compositing/Composite/Source/Composite/Private/Passes/CompositePassColorGrade.cpp
+++ b/Engine/Plugins/Experimental/Compositing/Composite/Source/Composite/Private/Passes/CompositePassColorGrade.cpp
@@ -60,7 +60,7 @@ public:
 		: FCompositeCorePassProxy(MoveTemp(InPassDeclaredInputs))
 	{ }
 
-	UE::CompositeCore::FPassOutput Add(FRDGBuilder& GraphBuilder, const FSceneView& InView, const UE::CompositeCore::FPassInputArray& Inputs, const UE::CompositeCore::FPassContext& PassContext) const override
+	UE::CompositeCore::FPassTexture Add(FRDGBuilder& GraphBuilder, const FSceneView& InView, const UE::CompositeCore::FPassInputArray& Inputs, const UE::CompositeCore::FPassContext& PassContext) const override
 	{
 		RDG_EVENT_SCOPE_STAT(GraphBuilder, FCompositeColorGrade, "Composite.ColorGrade");
 		RDG_GPU_STAT_SCOPE(GraphBuilder, FCompositeColorGrade);
@@ -107,7 +107,7 @@ public:
 			Parameters
 		);
 
-		return UE::CompositeCore::FPassOutput{ MoveTemp(Output), Metadata, PassOutputOverride };
+		return UE::CompositeCore::FPassTexture{ MoveTemp(Output), Metadata };
 	}
 
 	FCompositeTemperatureSettings TemperatureSettings;
diff --git a/Engine/Plugins/Experimental/Compositing/Composite/Source/Composite/Private/Passes/CompositePassDeltaKeyer.cpp b/Engine/Plugins/Experimental/Compositing/Composite/Source/Composite/Private/Passes/CompositePassDeltaKeyer.cpp
--- a/Engine/Plugins/Experimental/Compositing/Composite/Source/Composite/Private/Passes/CompositePassDeltaKeyer.cpp
+++ b/Engine/Plugins/Experimental/Compositing/Composite/Source/Composite/Private/Passes/CompositePassDeltaKeyer.cpp
@@ -54,7 +54,7 @@ namespace UE
 		public:
 			using FCompositeCorePassProxy::FCompositeCorePassProxy;
 
-			FPassOutput Add(FRDGBuilder& GraphBuilder, const FSceneView& InView, const FPassInputArray& Inputs, const FPassContext& PassContext) const override
+			FPassTexture Add(FRDGBuilder& GraphBuilder, const FSceneView& InView, const FPassInputArray& Inputs, const FPassContext& PassContext) const override
 			{
 				RDG_EVENT_SCOPE_STAT(GraphBuilder, FCompositeDeltaKeyer, "Composite.DeltaKeyer");
 				RDG_GPU_STAT_SCOPE(GraphBuilder, FCompositeDeltaKeyer);
@@ -118,7 +118,7 @@ namespace UE
 					FComputeShaderUtils::GetGroupCount(TextureSize, FCompositePassDeltaKeyerShader::ThreadGroupSize)
 				);
 
-				return UE::CompositeCore::FPassOutput{ MoveTemp(Output), Metadata, PassOutputOverride };
+				return UE::CompositeCore::FPassTexture{ MoveTemp(Output), Metadata };
 			}
 
 			FLinearColor KeyColor = FLinearColor::White;
diff --git a/Engine/Plugins/Experimental/Compositing/Composite/Source/Composite/Private/Passes/CompositePassDistortion.cpp b/Engine/Plugins/Experimental/Compositing/Composite/Source/Composite/Private/Passes/CompositePassDistortion.cpp
--- a/Engine/Plugins/Experimental/Compositing/Composite/Source/Composite/Private/Passes/CompositePassDistortion.cpp
+++ b/Engine/Plugins/Experimental/Compositing/Composite/Source/Composite/Private/Passes/CompositePassDistortion.cpp
@@ -50,7 +50,7 @@ public:
 
 	}
 
-	UE::CompositeCore::FPassOutput Add(FRDGBuilder& GraphBuilder, const FSceneView& InView, const UE::CompositeCore::FPassInputArray& Inputs, const UE::CompositeCore::FPassContext& PassContext) const override
+	UE::CompositeCore::FPassTexture Add(FRDGBuilder& GraphBuilder, const FSceneView& InView, const UE::CompositeCore::FPassInputArray& Inputs, const UE::CompositeCore::FPassContext& PassContext) const override
 	{
 		RDG_EVENT_SCOPE_STAT(GraphBuilder, FCompositeDistortion, "Composite.Distortion");
 		RDG_GPU_STAT_SCOPE(GraphBuilder, FCompositeDistortion);
@@ -130,11 +130,11 @@ public:
 
 			Metadata.bDistorted = (Distortion == ECompositeDistortion::Distort);
 			
-			return UE::CompositeCore::FPassOutput{ MoveTemp(Output), Metadata, PassOutputOverride };
+			return UE::CompositeCore::FPassTexture{ MoveTemp(Output), Metadata };
 		}
 		else
 		{
-			return UE::CompositeCore::FPassOutput{ Output.IsValid() ? Output : Input, Metadata, PassOutputOverride };
+			return UE::CompositeCore::FPassTexture{ Output.IsValid() ? Output : Input, Metadata };
 		}
 	}
 
diff --git a/Engine/Plugins/Experimental/Compositing/Composite/Source/Composite/Private/Passes/CompositePassMaterial.cpp b/Engine/Plugins/Experimental/Compositing/Composite/Source/Composite/Private/Passes/CompositePassMaterial.cpp
--- a/Engine/Plugins/Experimental/Compositing/Composite/Source/Composite/Private/Passes/CompositePassMaterial.cpp
+++ b/Engine/Plugins/Experimental/Compositing/Composite/Source/Composite/Private/Passes/CompositePassMaterial.cpp
@@ -19,7 +19,7 @@ public:
 	/* Post-process material weak pointer. */
 	TWeakObjectPtr<UMaterialInterface> Material;
 
-	UE::CompositeCore::FPassOutput Add(FRDGBuilder& GraphBuilder, const FSceneView& InView, const UE::CompositeCore::FPassInputArray& Inputs, const UE::CompositeCore::FPassContext& PassContext) const override
+	UE::CompositeCore::FPassTexture Add(FRDGBuilder& GraphBuilder, const FSceneView& InView, const UE::CompositeCore::FPassInputArray& Inputs, const UE::CompositeCore::FPassContext& PassContext) const override
 	{
 		RDG_EVENT_SCOPE_STAT(GraphBuilder, FCompositeMaterial, "Composite.Material");
 		RDG_GPU_STAT_SCOPE(GraphBuilder, FCompositeMaterial);
@@ -37,10 +37,10 @@ public:
 		{
 			FScreenPassTexture Output = AddPostProcessMaterialPass(GraphBuilder, InView, PostInputs, Material.GetEvenIfUnreachable());
 
-			return UE::CompositeCore::FPassOutput{ MoveTemp(Output), Metadata, PassOutputOverride };
+			return UE::CompositeCore::FPassTexture{ MoveTemp(Output), Metadata };
 		}
 
-		return UE::CompositeCore::FPassOutput{ PostInputs.ReturnUntouchedSceneColorForPostProcessing(GraphBuilder), Metadata, PassOutputOverride };
+		return UE::CompositeCore::FPassTexture{ PostInputs.ReturnUntouchedSceneColorForPostProcessing(GraphBuilder), Metadata };
 	}
 };
 
@@ -67,3 +67,10 @@ bool UCompositePassMaterial::GetProxy(const UE::CompositeCore::FPassInputDecl& I
 	return false;
 }
 
+bool UCompositePassMaterial::NeedsSceneTextures() const
+{
+	// We conservatively mark any post-processing materials as potentially using scene textures.
+	// TBD if material inputs could be individually checked instead.
+	return true;
+}
+
diff --git a/Engine/Plugins/Experimental/Compositing/Composite/Source/Composite/Private/Passes/CompositePassOpenColorIO.cpp b/Engine/Plugins/Experimental/Compositing/Composite/Source/Composite/Private/Passes/CompositePassOpenColorIO.cpp
--- a/Engine/Plugins/Experimental/Compositing/Composite/Source/Composite/Private/Passes/CompositePassOpenColorIO.cpp
+++ b/Engine/Plugins/Experimental/Compositing/Composite/Source/Composite/Private/Passes/CompositePassOpenColorIO.cpp
@@ -20,7 +20,7 @@ public:
 
 	}
 
-	UE::CompositeCore::FPassOutput Add(FRDGBuilder& GraphBuilder, const FSceneView& InView, const UE::CompositeCore::FPassInputArray& Inputs, const UE::CompositeCore::FPassContext& PassContext) const override
+	UE::CompositeCore::FPassTexture Add(FRDGBuilder& GraphBuilder, const FSceneView& InView, const UE::CompositeCore::FPassInputArray& Inputs, const UE::CompositeCore::FPassContext& PassContext) const override
 	{
 		RDG_EVENT_SCOPE_STAT(GraphBuilder, FCompositeOpenColorIO, "Composite.OpenColorIO");
 		RDG_GPU_STAT_SCOPE(GraphBuilder, FCompositeOpenColorIO);
@@ -45,7 +45,7 @@ public:
 			PassResources
 		);
 
-		return UE::CompositeCore::FPassOutput{ MoveTemp(Output), Metadata, PassOutputOverride };
+		return UE::CompositeCore::FPassTexture{ MoveTemp(Output), Metadata };
 	}
 
 	FOpenColorIORenderPassResources PassResources;
diff --git a/Engine/Plugins/Experimental/Compositing/Composite/Source/Composite/Private/Passes/CompositePassSMAA.cpp b/Engine/Plugins/Experimental/Compositing/Composite/Source/Composite/Private/Passes/CompositePassSMAA.cpp
--- a/Engine/Plugins/Experimental/Compositing/Composite/Source/Composite/Private/Passes/CompositePassSMAA.cpp
+++ b/Engine/Plugins/Experimental/Compositing/Composite/Source/Composite/Private/Passes/CompositePassSMAA.cpp
@@ -22,7 +22,7 @@ namespace UE
 		public:
 			using FCompositeCorePassProxy::FCompositeCorePassProxy;
 
-			FPassOutput Add(FRDGBuilder& GraphBuilder, const FSceneView& InView, const FPassInputArray& Inputs, const FPassContext& PassContext) const override
+			FPassTexture Add(FRDGBuilder& GraphBuilder, const FSceneView& InView, const FPassInputArray& Inputs, const FPassContext& PassContext) const override
 			{
 				RDG_EVENT_SCOPE_STAT(GraphBuilder, FCompositeCoreSMAA, "CompositeCore.SMAA");
 				RDG_GPU_STAT_SCOPE(GraphBuilder, FCompositeCoreSMAA);
@@ -62,7 +62,7 @@ namespace UE
 					Output = CompositeCore::Private::AddDisplayTransformPass(GraphBuilder, InView, Output, Inputs.OverrideOutput, bIsForward);
 				}
 
-				return FPassOutput{ MoveTemp(Output), Metadata, PassOutputOverride };
+				return FPassTexture{ MoveTemp(Output), Metadata };
 			}
 
 			/** Optional r.SMAA.Quality setting override.  */
diff --git a/Engine/Plugins/Experimental/Compositing/Composite/Source/Composite/Public/Layers/CompositeLayerBase.h b/Engine/Plugins/Experimental/Compositing/Composite/Source/Composite/Public/Layers/CompositeLayerBase.h
--- a/Engine/Plugins/Experimental/Compositing/Composite/Source/Composite/Public/Layers/CompositeLayerBase.h
+++ b/Engine/Plugins/Experimental/Compositing/Composite/Source/Composite/Public/Layers/CompositeLayerBase.h
@@ -29,7 +29,10 @@ public:
 		bool bIsSolo = false;
 		
 		/** Flag to indicate first pass */
-		bool bIsFirstPass = false;
+		bool bIsFirstPass = true;
+
+		/** Flag to indicate that a layer is independent of scene textures. Work from scene-independent layers can be shared between post-processing locations. */
+		bool bNeedsSceneTextures = false;
 		
 		/** List of external textures */
 		TArray<UE::CompositeCore::FExternalTexture> ExternalTextures;
@@ -72,6 +75,10 @@ public:
 	ECompositeCoreMergeOp Operation = ECompositeCoreMergeOp::Over;
 
 protected:
+
+	/** Convenience function to register post-processing child passes. */
+	void AddChildPasses(UE::CompositeCore::FPassInputDecl& InBasePassInput, FTraversalContext& InContext, FSceneRenderingBulkObjectAllocator& InFrameAllocator, TArrayView<const TObjectPtr<UCompositePassBase>> InPasses) const;
+
 	/** Fixed number of inputs used by layer merge passes. */
 	static constexpr int32 FixedNumLayerInputs = 2;
 
diff --git a/Engine/Plugins/Experimental/Compositing/Composite/Source/Composite/Public/Passes/CompositePassBase.h b/Engine/Plugins/Experimental/Compositing/Composite/Source/Composite/Public/Passes/CompositePassBase.h
--- a/Engine/Plugins/Experimental/Compositing/Composite/Source/Composite/Public/Passes/CompositePassBase.h
+++ b/Engine/Plugins/Experimental/Compositing/Composite/Source/Composite/Public/Passes/CompositePassBase.h
@@ -33,7 +33,10 @@ public:
 	UE_API virtual bool GetProxy(const UE::CompositeCore::FPassInputDecl& InputDecl, FSceneRenderingBulkObjectAllocator& InFrameAllocator, FCompositeCorePassProxy*& OutProxy) const { return false; };
 
 	/** Returns true if the pass is returns a valid proxy. */
-	UE_API virtual bool IsActive() const { return bIsEnabled; }
+	virtual bool IsActive() const { return bIsEnabled; }
+
+	/** Returns true if the pass (may) need post-processing scene textures. */
+	virtual bool NeedsSceneTextures() const { return false; }
 
 	/** Whether or not the pass is active. */
 	UPROPERTY(EditAnywhere, BlueprintReadWrite, Category = "", meta = (DisplayPriority = "1"))
diff --git a/Engine/Plugins/Experimental/Compositing/Composite/Source/Composite/Public/Passes/CompositePassMaterial.h b/Engine/Plugins/Experimental/Compositing/Composite/Source/Composite/Public/Passes/CompositePassMaterial.h
--- a/Engine/Plugins/Experimental/Compositing/Composite/Source/Composite/Public/Passes/CompositePassMaterial.h
+++ b/Engine/Plugins/Experimental/Compositing/Composite/Source/Composite/Public/Passes/CompositePassMaterial.h
@@ -25,6 +25,8 @@ public:
 
 	UE_API virtual bool GetProxy(const UE::CompositeCore::FPassInputDecl& InputDecl, FSceneRenderingBulkObjectAllocator& InFrameAllocator, FCompositeCorePassProxy*& OutProxy) const override;
 
+	UE_API virtual bool NeedsSceneTextures() const override;
+
 public:
 	/** Post-process material to execute. Input0 is connected to SceneTexture's PostprocessInput0. */
 	UPROPERTY(EditAnywhere, BlueprintReadWrite, Category = "")
diff --git a/Engine/Plugins/Experimental/Compositing/CompositeCore/Source/CompositeCore/Private/CompositeCoreSceneViewExtension.cpp b/Engine/Plugins/Experimental/Compositing/CompositeCore/Source/CompositeCore/Private/CompositeCoreSceneViewExtension.cpp
--- a/Engine/Plugins/Experimental/Compositing/CompositeCore/Source/CompositeCore/Private/CompositeCoreSceneViewExtension.cpp
+++ b/Engine/Plugins/Experimental/Compositing/CompositeCore/Source/CompositeCore/Private/CompositeCoreSceneViewExtension.cpp
@@ -135,16 +135,16 @@ void FCompositeCoreSceneViewExtension::SetRenderWork_GameThread(UE::CompositeCor
 	bHasCustomRenderWork = !InWork.IsEmpty();
 
 	ENQUEUE_RENDER_COMMAND(CopyCompositeCoreRenderWork)(
-		[InWork = MoveTemp(InWork), WeakThis = AsWeak()](FRHICommandListImmediate& RHICmdList) mutable
+		[FrameWork = MoveTemp(InWork), WeakThis = AsWeak()](FRHICommandListImmediate& RHICmdList) mutable
 		{
 			TSharedPtr<FCompositeCoreSceneViewExtension> SVE = StaticCastSharedPtr<FCompositeCoreSceneViewExtension>(WeakThis.Pin());
 			if (SVE.IsValid())
 			{
 				SVE->ExternalInputs_RenderThread.Reset();
 
-				for (int32 Index=0; Index < InWork.ExternalInputs.Num(); ++Index)
+				for (int32 Index=0; Index < FrameWork.ExternalInputs.Num(); ++Index)
 				{
-					const UE::CompositeCore::FExternalTexture& ExternalInput = InWork.ExternalInputs[Index];
+					const UE::CompositeCore::FExternalTexture& ExternalInput = FrameWork.ExternalInputs[Index];
 					TStrongObjectPtr<UTexture> ExternalTexture = ExternalInput.Texture.Pin();
 					
 					if (ExternalTexture.IsValid())
@@ -160,7 +160,7 @@ void FCompositeCoreSceneViewExtension::SetRenderWork_GameThread(UE::CompositeCor
 					}
 				}
 
-				SVE->RenderWork_RenderThread = MoveTemp(InWork);
+				SVE->RenderWork_RenderThread = MoveTemp(FrameWork);
 			}
 		});
 }
@@ -336,6 +336,21 @@ void FCompositeCoreSceneViewExtension::BeginRenderViewFamily(FSceneViewFamily& I
 
 void FCompositeCoreSceneViewExtension::PreRenderViewFamily_RenderThread(FRDGBuilder& GraphBuilder, FSceneViewFamily& InViewFamily)
 {
+	// To prevent cache reallocations each frame, we keep the existing views unless they start exceeding a reasonable number.
+	constexpr int32 MaxNumCachedViews = 16;
+
+	if (CachedOutputsPerView_RenderThread.Num() <= MaxNumCachedViews)
+	{
+		for (auto& Pair : CachedOutputsPerView_RenderThread)
+		{
+			// Reset each pre-cached view's output
+			Pair.Value.Reset();
+		}
+	}
+	else
+	{
+		CachedOutputsPerView_RenderThread.Reset();
+	}
 }
 
 void FCompositeCoreSceneViewExtension::PreRenderView_RenderThread(FRDGBuilder& GraphBuilder, FSceneView& InView)
@@ -367,8 +382,8 @@ void FCompositeCoreSceneViewExtension::PrePostProcessPass_RenderThread(FRDGBuild
 		PassInput.Texture = FScreenPassTexture{ GraphBuilder.RegisterExternalTexture(*CompositeRenderPassPtr) };
 
 		// Apply FXAA, with additional fwd/inv display transform passes
-		const FPassOutput Output = FXAAPassProxy.Add(GraphBuilder, InView, PassInputs, {});
-		const FScreenPassTexture& ScreenTexOutput = Output.Resource.Texture;
+		const FPassTexture Output = FXAAPassProxy.Add(GraphBuilder, InView, PassInputs, {});
+		const FScreenPassTexture& ScreenTexOutput = Output.Texture;
 
 		// Extract the result back into the composite render target
 		*CompositeRenderPassPtr = GraphBuilder.ConvertToExternalTexture(ScreenTexOutput.Texture);
@@ -433,7 +448,7 @@ void FCompositeCoreSceneViewExtension::ApplyPreprocessing(FRDGBuilder& GraphBuil
 			const FCompositeCorePassProxy* Pass = Passes[PassIndex];
 
 			FRDGTextureRef ExternalRDG = nullptr;
-			const TOptional<UE::CompositeCore::ResourceId>& OverrrideOutput = Pass->GetDeclareOutputOverride();
+			const TOptional<UE::CompositeCore::ResourceId>& OverrrideOutput = Pass->GetDeclaredOutputOverride();
 			
 			if (OverrrideOutput.IsSet())
 			{
@@ -456,7 +471,9 @@ void FCompositeCoreSceneViewExtension::ApplyPreprocessing(FRDGBuilder& GraphBuil
 			}
 
 			// Register pass and update output
-			const FPassOutput Output = Pass->Add(GraphBuilder, InView, PassInputs, PassContext);
+			const FPassTexture Output = Pass->Add(GraphBuilder, InView, PassInputs, PassContext);
+
+			// Note: we do not cache the output to CachedOutputsPerView_RenderThread since we only support cached outputs in post-processing for now.
 
 			// Manually transition the external render target back to a shader resource.
 			if (ExternalRDG != nullptr)
@@ -468,7 +485,7 @@ void FCompositeCoreSceneViewExtension::ApplyPreprocessing(FRDGBuilder& GraphBuil
 			}
 
 			// Update next pass input
-			PassInputs[0] = Output.Resource;
+			PassInputs[0] = Output;
 		}
 	}
 }
@@ -503,33 +520,6 @@ TSortedMap<int32, UE::CompositeCore::FPassInput> FCompositeCoreSceneViewExtensio
 	return InputOverrides;
 }
 
-void FCompositeCoreSceneViewExtension::UpdateNextPassInputs(
-	FRDGBuilder& GraphBuilder,
-	const UE::CompositeCore::FPassOutput& InOutput,
-	const int32 BindingIndex,
-	UE::CompositeCore::FPassInputArray& InOutInputs,
-	FExternalTextureMap& InOutExternalTextures
-)
-{
-	using namespace UE::CompositeCore;
-
-	if (InOutput.Override.IsSet())
-	{
-		const ResourceId& ExternalId = InOutput.Override.GetValue();
-
-		FPassInput* ExternalTexture = InOutExternalTextures.Find(ExternalId);
-		if (ensureMsgf(ExternalTexture != nullptr, TEXT("Unexpected missing external texture override as output.")))
-		{
-			*ExternalTexture = InOutput.Resource;
-		}
-	}
-	else
-	{
-		// Update input binding for the next pass(es)
-		InOutInputs[BindingIndex] = InOutput.Resource;
-	}
-}
-
 bool FCompositeCoreSceneViewExtension::ApplyPasses_Recursive(
 	FRDGBuilder& GraphBuilder,
 	const FSceneView& InView,
@@ -539,7 +529,7 @@ bool FCompositeCoreSceneViewExtension::ApplyPasses_Recursive(
 	const TArray<const FCompositeCorePassProxy*> InPasses,
 	FExternalTextureMap& ExternalTextures,
 	int32 RecursionLevel,
-	UE::CompositeCore::FPassOutput& Output)
+	UE::CompositeCore::FPassTexture& Output)
 {
 	using namespace UE::CompositeCore;
 
@@ -550,6 +540,7 @@ bool FCompositeCoreSceneViewExtension::ApplyPasses_Recursive(
 
 	// Default pass inputs
 	FPassInputArray BasePassInputs = Inputs;
+	const TSortedMap<const FCompositeCorePassProxy*, const UE::CompositeCore::FPassTexture>& CachedOutputs = CachedOutputsPerView_RenderThread.FindOrAdd(InView.GetViewKey());
 	
 	// Iterate over all passes
 	for (int32 PassIndex = 0; PassIndex < InPasses.Num(); ++PassIndex)
@@ -571,11 +562,16 @@ bool FCompositeCoreSceneViewExtension::ApplyPasses_Recursive(
 			if (DeclaredInput.IsType<const FCompositeCorePassProxy*>())
 			{
 				const FCompositeCorePassProxy* ChildPassProxy = DeclaredInput.Get<const FCompositeCorePassProxy*>();
-				UE::CompositeCore::FPassOutput ChildPassOutput;
+				UE::CompositeCore::FPassTexture ChildPassOutput;
 
-				if (ApplyPasses_Recursive(GraphBuilder, InView, PassInputs, OriginalInputs, PassContext, { ChildPassProxy }, ExternalTextures, RecursionLevel + 1, ChildPassOutput))
+				// We first check if the input proxy has already produced an output (and let RDG do the rest!)
+				if (const UE::CompositeCore::FPassTexture* PreProducedOutput = CachedOutputs.Find(ChildPassProxy))
+				{
+					PassInputs[InputIndex] = *PreProducedOutput;
+				}
+				else if(ApplyPasses_Recursive(GraphBuilder, InView, PassInputs, OriginalInputs, PassContext, { ChildPassProxy }, ExternalTextures, RecursionLevel + 1, ChildPassOutput))
 				{
-					UpdateNextPassInputs(GraphBuilder, ChildPassOutput, InputIndex, PassInputs, ExternalTextures);
+					PassInputs[InputIndex] = ChildPassOutput;
 				}
 			}
 			// If an external texture resource is expected, connect the current input index with the resolved identifier
@@ -630,8 +626,11 @@ bool FCompositeCoreSceneViewExtension::ApplyPasses_Recursive(
 
 		// Register pass and update output
 		Output = Pass->Add(GraphBuilder, InView, PassInputs, PassContext);
+		
+		// Cache the pass output per view
+		CachedOutputsPerView_RenderThread.FindOrAdd(InView.GetViewKey()).Add(Pass, Output);
 
-		UpdateNextPassInputs(GraphBuilder, Output, 0, BasePassInputs, ExternalTextures);
+		BasePassInputs[0] = Output;
 	}
 
 	return true;
@@ -651,14 +650,12 @@ FScreenPassTexture FCompositeCoreSceneViewExtension::PostProcessWork_RenderThrea
 	PassContext.bOutputSceneColor = false;
 
 	FPassInputArray ResolvedInputs(GraphBuilder, InView, Inputs, InLocation);
-	FPassOutput Output;
+	FPassTexture Output;
 	
 	constexpr int32 RecursionLevel = 0;
-	ApplyPasses_Recursive(GraphBuilder, InView, ResolvedInputs, ResolvedInputs, PassContext, Passes, ExternalTextures, RecursionLevel, Output);
-
-	if (Output.Resource.Texture.IsValid() && !Output.Override.IsSet())
+	if (ApplyPasses_Recursive(GraphBuilder, InView, ResolvedInputs, ResolvedInputs, PassContext, Passes, ExternalTextures, RecursionLevel, Output))
 	{
-		return MoveTemp(Output.Resource.Texture);
+		return MoveTemp(Output.Texture);
 	}
 	else
 	{
diff --git a/Engine/Plugins/Experimental/Compositing/CompositeCore/Source/CompositeCore/Private/CompositeCoreSceneViewExtension.h b/Engine/Plugins/Experimental/Compositing/CompositeCore/Source/CompositeCore/Private/CompositeCoreSceneViewExtension.h
--- a/Engine/Plugins/Experimental/Compositing/CompositeCore/Source/CompositeCore/Private/CompositeCoreSceneViewExtension.h
+++ b/Engine/Plugins/Experimental/Compositing/CompositeCore/Source/CompositeCore/Private/CompositeCoreSceneViewExtension.h
@@ -92,15 +92,6 @@ private:
 
 	/** Convenience function to create an override map of external textures. */
 	FExternalTextureMap CreateExternalTextureMap(FRDGBuilder& GraphBuilder, const FSceneView& InView) const;
-	
-	/** Convenience function to update the next pass inputs from the last pass' output. */
-	void UpdateNextPassInputs(
-		FRDGBuilder& GraphBuilder,
-		const UE::CompositeCore::FPassOutput& InOutput,
-		const int32 BindingIndex,
-		UE::CompositeCore::FPassInputArray& InOutInputs,
-		FExternalTextureMap& InOutExternalTextures
-	);
 
 	/** Recursive pass application. */
 	bool ApplyPasses_Recursive(
@@ -112,7 +103,7 @@ private:
 		const TArray<const FCompositeCorePassProxy*> InPasses,
 		FExternalTextureMap& InOutExternalTextures,
 		int32 RecursionLevel,
-		UE::CompositeCore::FPassOutput& Output
+		UE::CompositeCore::FPassTexture& Output
 	);
 
 	/** Callback for processing passes. */
@@ -123,10 +114,10 @@ private:
 	TSet<TWeakObjectPtr<UPrimitiveComponent>> CompositePrimitives;
 
 	/** Array of external texture inputs. */
-	TMap<UE::CompositeCore::ResourceId, UE::CompositeCore::FExternalRenderTarget> ExternalInputs_RenderThread;
+	TSortedMap<UE::CompositeCore::ResourceId, UE::CompositeCore::FExternalRenderTarget> ExternalInputs_RenderThread;
 
 	/** Custom render pass render targets for each active view .*/
-	TMap<uint32, TRefCountPtr<IPooledRenderTarget>> CustomRenderTargetPerView_RenderThread;
+	TSortedMap<uint32, TRefCountPtr<IPooledRenderTarget>> CustomRenderTargetPerView_RenderThread;
 
 	/** Custom render work. */
 	TOptional<UE::CompositeCore::FRenderWork> RenderWork_RenderThread;
@@ -136,5 +127,15 @@ private:
 
 	/** Used to keep the SVE active if any custom render work is present.*/
 	bool bHasCustomRenderWork = false;
+
+	/**
+	* Map of cached proxy output textures per view (post-processing only currently).
+	* 
+	* Passes can then refer to a previous proxy as an input type, and if that proxy has
+	* already been executed its output will be used directly without needing to duplicate work.
+	* This is possible since RDG automatically knows to connect & preserve the texture resource
+	* when it is later used as an input of a subsequent pass.
+	*/
+	mutable TSortedMap<uint32, TSortedMap<const FCompositeCorePassProxy*, const UE::CompositeCore::FPassTexture>> CachedOutputsPerView_RenderThread;
 };
 
diff --git a/Engine/Plugins/Experimental/Compositing/CompositeCore/Source/CompositeCore/Private/Passes/CompositeCorePassFXAAProxy.cpp b/Engine/Plugins/Experimental/Compositing/CompositeCore/Source/CompositeCore/Private/Passes/CompositeCorePassFXAAProxy.cpp
--- a/Engine/Plugins/Experimental/Compositing/CompositeCore/Source/CompositeCore/Private/Passes/CompositeCorePassFXAAProxy.cpp
+++ b/Engine/Plugins/Experimental/Compositing/CompositeCore/Source/CompositeCore/Private/Passes/CompositeCorePassFXAAProxy.cpp
@@ -79,7 +79,7 @@ namespace UE
 			}
 		}
 
-		FPassOutput FFXAAPassProxy::Add(FRDGBuilder& GraphBuilder, const FSceneView& InView, const FPassInputArray& Inputs, const FPassContext& PassContext) const
+		FPassTexture FFXAAPassProxy::Add(FRDGBuilder& GraphBuilder, const FSceneView& InView, const FPassInputArray& Inputs, const FPassContext& PassContext) const
 		{
 			RDG_EVENT_SCOPE_STAT(GraphBuilder, FCompositeCoreFXAA, "CompositeCore.FXAA");
 			RDG_GPU_STAT_SCOPE(GraphBuilder, FCompositeCoreFXAA);
@@ -118,7 +118,7 @@ namespace UE
 				Output = Private::AddDisplayTransformPass(GraphBuilder, InView, Output, Inputs.OverrideOutput, bIsForward);
 			}
 
-			return FPassOutput{ MoveTemp(Output), Metadata, PassOutputOverride };
+			return FPassTexture{ MoveTemp(Output), Metadata };
 		}
 	}
 }
diff --git a/Engine/Plugins/Experimental/Compositing/CompositeCore/Source/CompositeCore/Private/Passes/CompositeCorePassMergeProxy.cpp b/Engine/Plugins/Experimental/Compositing/CompositeCore/Source/CompositeCore/Private/Passes/CompositeCorePassMergeProxy.cpp
--- a/Engine/Plugins/Experimental/Compositing/CompositeCore/Source/CompositeCore/Private/Passes/CompositeCorePassMergeProxy.cpp
+++ b/Engine/Plugins/Experimental/Compositing/CompositeCore/Source/CompositeCore/Private/Passes/CompositeCorePassMergeProxy.cpp
@@ -3,7 +3,9 @@
 #include "Passes/CompositeCorePassMergeProxy.h"
 
 #include "CompositeCoreSettings.h"
+#include "CompositeCoreModule.h"
 
+#include "Logging/StructuredLog.h"
 #include "PixelShaderUtils.h"
 #include "PostProcess/LensDistortion.h"
 #include "PostProcess/PostProcessMaterialInputs.h"
@@ -114,19 +116,25 @@ namespace UE
 			return Parameters;
 		}
 
-		FMergePassProxy::FMergePassProxy(FPassInputDeclArray InPassDeclaredInputs, ECompositeCoreMergeOp InMergeOp, ELensDistortionHandling bInUseLensDistortion)
+		FMergePassProxy::FMergePassProxy(FPassInputDeclArray InPassDeclaredInputs, ECompositeCoreMergeOp InMergeOp, const TCHAR* InParentName, ELensDistortionHandling bInUseLensDistortion)
 			: FCompositeCorePassProxy(MoveTemp(InPassDeclaredInputs))
 			, MergeOp(InMergeOp)
+			, ParentName(InParentName)
 			, bUseLensDistortion(bInUseLensDistortion)
 		{
 		}
 
-		FPassOutput FMergePassProxy::Add(FRDGBuilder& GraphBuilder, const FSceneView& InView, const FPassInputArray& Inputs, const FPassContext& PassContext) const
+		FPassTexture FMergePassProxy::Add(FRDGBuilder& GraphBuilder, const FSceneView& InView, const FPassInputArray& Inputs, const FPassContext& PassContext) const
 		{
-			RDG_EVENT_SCOPE_STAT(GraphBuilder, FCompositeCoreMerge, "CompositeCore.Merge");
+			check(ValidateInputs(Inputs));
+			RDG_EVENT_SCOPE_STAT(GraphBuilder, FCompositeCoreMerge, "CompositeCore.Merge (%s) %dx%d",
+				ParentName,
+				Inputs[0].Texture.ViewRect.Width(),
+				Inputs[0].Texture.ViewRect.Height()
+			);
 			RDG_GPU_STAT_SCOPE(GraphBuilder, FCompositeCoreMerge);
 
-			check(ValidateInputs(Inputs));
+			UE_LOG(LogCompositeCore, VeryVerbose, TEXT("CompositeCore.Merge: %s"), ParentName);
 
 			FGlobalShaderMap* GlobalShaderMap = GetGlobalShaderMap(InView.GetFeatureLevel());
 			const FSceneViewFamily* Family = InView.Family;
@@ -216,7 +224,7 @@ namespace UE
 				Output.ViewRect
 			);
 
-			return FPassOutput{ MoveTemp(Output), MoveTemp(OutMetadata), PassOutputOverride };
+			return FPassTexture{ MoveTemp(Output), MoveTemp(OutMetadata) };
 		}
 	}
 }
diff --git a/Engine/Plugins/Experimental/Compositing/CompositeCore/Source/CompositeCore/Private/Passes/CompositeCorePassProxy.cpp b/Engine/Plugins/Experimental/Compositing/CompositeCore/Source/CompositeCore/Private/Passes/CompositeCorePassProxy.cpp
--- a/Engine/Plugins/Experimental/Compositing/CompositeCore/Source/CompositeCore/Private/Passes/CompositeCorePassProxy.cpp
+++ b/Engine/Plugins/Experimental/Compositing/CompositeCore/Source/CompositeCore/Private/Passes/CompositeCorePassProxy.cpp
@@ -2,6 +2,7 @@
 
 #include "Passes/CompositeCorePassProxy.h"
 
+#include "CompositeCoreModule.h"
 #include "Passes/CompositeCorePassMergeProxy.h"
 
 #include "FXRenderingUtils.h"
@@ -20,12 +21,6 @@ namespace UE
 {
 	namespace CompositeCore
 	{
-		FPassOutput::FPassOutput(FScreenPassTexture InTexture, FResourceMetadata InMetadata, TOptional<ResourceId> InOverride)
-			: Resource({ MoveTemp(InTexture), MoveTemp(InMetadata) })
-			, Override(InOverride)
-		{
-		}
-
 		FBuiltInRenderPassOptions::FBuiltInRenderPassOptions()
 			: ViewUserFlagsOverride{}
 			, bEnableUnlitViewmode{true}
@@ -94,7 +89,7 @@ namespace UE
 
 FCompositeCorePassProxy::FCompositeCorePassProxy(UE::CompositeCore::FPassInputDeclArray InPassDeclaredInputs)
 	: PassDeclaredInputs(MoveTemp(InPassDeclaredInputs))
-	, PassOutputOverride()
+	, PassDeclaredOutputOverride()
 { }
 
 FScreenPassRenderTarget FCompositeCorePassProxy::CreateOutputRenderTarget(FRDGBuilder& GraphBuilder,
diff --git a/Engine/Plugins/Experimental/Compositing/CompositeCore/Source/CompositeCore/Public/Passes/CompositeCorePassFXAAProxy.h b/Engine/Plugins/Experimental/Compositing/CompositeCore/Source/CompositeCore/Public/Passes/CompositeCorePassFXAAProxy.h
--- a/Engine/Plugins/Experimental/Compositing/CompositeCore/Source/CompositeCore/Public/Passes/CompositeCorePassFXAAProxy.h
+++ b/Engine/Plugins/Experimental/Compositing/CompositeCore/Source/CompositeCore/Public/Passes/CompositeCorePassFXAAProxy.h
@@ -21,7 +21,7 @@ namespace UE
 		public:
 			using FCompositeCorePassProxy::FCompositeCorePassProxy;
 
-			UE_API FPassOutput Add(FRDGBuilder& GraphBuilder, const FSceneView& InView, const FPassInputArray& Inputs, const FPassContext& PassContext) const override;
+			UE_API FPassTexture Add(FRDGBuilder& GraphBuilder, const FSceneView& InView, const FPassInputArray& Inputs, const FPassContext& PassContext) const override;
 
 			/** Optional r.FXAA.Quality setting override.  */
 			TOptional<int32> QualityOverride = {};
diff --git a/Engine/Plugins/Experimental/Compositing/CompositeCore/Source/CompositeCore/Public/Passes/CompositeCorePassMergeProxy.h b/Engine/Plugins/Experimental/Compositing/CompositeCore/Source/CompositeCore/Public/Passes/CompositeCorePassMergeProxy.h
--- a/Engine/Plugins/Experimental/Compositing/CompositeCore/Source/CompositeCore/Public/Passes/CompositeCorePassMergeProxy.h
+++ b/Engine/Plugins/Experimental/Compositing/CompositeCore/Source/CompositeCore/Public/Passes/CompositeCorePassMergeProxy.h
@@ -49,13 +49,16 @@ namespace UE
 			using FCompositeCorePassProxy::FCompositeCorePassProxy;
 
 			/** Constructor */
-			UE_API FMergePassProxy(FPassInputDeclArray InPassDeclaredInputs, ECompositeCoreMergeOp InMergeOp, ELensDistortionHandling bInUseLensDistortion = ELensDistortionHandling::Enabled);
+			UE_API FMergePassProxy(FPassInputDeclArray InPassDeclaredInputs, ECompositeCoreMergeOp InMergeOp, const TCHAR* InParentName, ELensDistortionHandling bInUseLensDistortion = ELensDistortionHandling::Enabled);
 
-			UE_API FPassOutput Add(FRDGBuilder& GraphBuilder, const FSceneView& InView, const FPassInputArray& Inputs, const FPassContext& PassContext) const override;
+			UE_API FPassTexture Add(FRDGBuilder& GraphBuilder, const FSceneView& InView, const FPassInputArray& Inputs, const FPassContext& PassContext) const override;
 
 			/** Merge operation. */
 			ECompositeCoreMergeOp MergeOp = ECompositeCoreMergeOp::Over;
 
+			/** Parent pass debug name. */
+			const TCHAR* ParentName = TEXT("None");
+
 			/** Flag to enable or disable application of engine-provided lens distortion LUTs, true by default. */
 			ELensDistortionHandling bUseLensDistortion = ELensDistortionHandling::Enabled;
 		};
diff --git a/Engine/Plugins/Experimental/Compositing/CompositeCore/Source/CompositeCore/Public/Passes/CompositeCorePassProxy.h b/Engine/Plugins/Experimental/Compositing/CompositeCore/Source/CompositeCore/Public/Passes/CompositeCorePassProxy.h
--- a/Engine/Plugins/Experimental/Compositing/CompositeCore/Source/CompositeCore/Public/Passes/CompositeCorePassProxy.h
+++ b/Engine/Plugins/Experimental/Compositing/CompositeCore/Source/CompositeCore/Public/Passes/CompositeCorePassProxy.h
@@ -168,25 +168,6 @@ namespace UE
 			TArray<UE::CompositeCore::FPassInput> Inputs;
 		};
 
-		/** Pass output definition. */
-		struct FPassOutput
-		{
-			/** Default operators */
-			FPassOutput() = default;
-			FPassOutput(FPassOutput&&) = default;
-			FPassOutput(const FPassOutput&) = default;
-			FPassOutput& operator=(FPassOutput&&) = default;
-			FPassOutput& operator=(const FPassOutput&) = default;
-
-			/** Constructor */
-			UE_API explicit FPassOutput(FScreenPassTexture InTexture, FResourceMetadata InMetadata, TOptional<ResourceId> InOverride = {});
-
-			/** Output pass resource. */
-			FPassTexture Resource = {};
-			/** Optional pass output target override, from the otherwise default scene color. */
-			TOptional<ResourceId> Override = {};
-		};
-
 		/** Pass parameter information. */
 		struct FPassContext
 		{
@@ -263,7 +244,7 @@ public:
 	UE_API virtual ~FCompositeCorePassProxy() = default;
 
 	/** Render-thread add pass method to override. */
-	virtual UE::CompositeCore::FPassOutput Add(
+	UE_API virtual UE::CompositeCore::FPassTexture Add(
 		FRDGBuilder& GraphBuilder,
 		const FSceneView& InView,
 		const UE::CompositeCore::FPassInputArray& Inputs,
@@ -282,22 +263,28 @@ public:
 		return PassDeclaredInputs[InputIndex];
 	}
 
+	/** Get pass input at specified index. */
+	UE::CompositeCore::FPassInputDecl& GetDeclaredInput(int32 InputIndex)
+	{
+		return PassDeclaredInputs[InputIndex];
+	}
+
 	/** Get declared output pass override. */
-	const TOptional<UE::CompositeCore::ResourceId>& GetDeclareOutputOverride() const
+	const TOptional<UE::CompositeCore::ResourceId>& GetDeclaredOutputOverride() const
 	{
-		return PassOutputOverride;	
+		return PassDeclaredOutputOverride;
 	}
 
 	/** Declare output pass override. */
 	void DeclareOutputOverride(UE::CompositeCore::ResourceId InExternalId)
 	{
-		PassOutputOverride = InExternalId;
+		PassDeclaredOutputOverride = InExternalId;
 	}
 
 	/** Reset output pass override. */
 	void ResetOutputOverride()
 	{
-		PassOutputOverride.Reset();
+		PassDeclaredOutputOverride.Reset();
 	}
 
 	/** Convenience function to create an output render target with the specified resolution. */
@@ -312,8 +299,8 @@ protected:
 	/** List of pass input types. */
 	UE::CompositeCore::FPassInputDeclArray PassDeclaredInputs;
 
-	/** Optional pass output override. */
-	TOptional<UE::CompositeCore::ResourceId> PassOutputOverride;
+	/** Optional declared output pass override. */
+	TOptional<UE::CompositeCore::ResourceId> PassDeclaredOutputOverride;
 };
 
 #undef UE_API
