CompositeCore: Commented out module dev build flags for convenience.

#prelight skip

diff --git a/Engine/Plugins/Experimental/Compositing/CompositeCore/Source/CompositeCore/CompositeCore.Build.cs b/Engine/Plugins/Experimental/Compositing/CompositeCore/Source/CompositeCore/CompositeCore.Build.cs
--- a/Engine/Plugins/Experimental/Compositing/CompositeCore/Source/CompositeCore/CompositeCore.Build.cs
+++ b/Engine/Plugins/Experimental/Compositing/CompositeCore/Source/CompositeCore/CompositeCore.Build.cs
@@ -43,5 +43,9 @@ public class CompositeCore : ModuleRules
 				}
 			);
 		}
+
+		//bUseUnity = false;
+		//OptimizeCode = CodeOptimization.Never;
+		//PCHUsage = PCHUsageMode.NoPCHs;
 	}
 }
