Composite: Fix RHI validation issue with passes operating on external render targets.

This could previously cause flickering.

#DisabledValidations nonUnityAndStaticAnalysis

diff --git a/Engine/Plugins/Experimental/Compositing/CompositeCore/Source/CompositeCore/Private/CompositeCoreSceneViewExtension.cpp b/Engine/Plugins/Experimental/Compositing/CompositeCore/Source/CompositeCore/Private/CompositeCoreSceneViewExtension.cpp
--- a/Engine/Plugins/Experimental/Compositing/CompositeCore/Source/CompositeCore/Private/CompositeCoreSceneViewExtension.cpp
+++ b/Engine/Plugins/Experimental/Compositing/CompositeCore/Source/CompositeCore/Private/CompositeCoreSceneViewExtension.cpp
@@ -452,7 +452,6 @@ void FCompositeCoreSceneViewExtension::ApplyPreprocessing(FRDGBuilder& GraphBuil
 		// Iterate over all passes
 		for (int32 PassIndex =0; PassIndex <Passes.Num(); ++PassIndex)
 		{
-			const bool bIsLastPass = (PassIndex == Passes.Num() - 1);
 			const FCompositeCorePassProxy* Pass = Passes[PassIndex];
 
 			FRDGTextureRef ExternalRDG = nullptr;
@@ -464,6 +463,8 @@ void FCompositeCoreSceneViewExtension::ApplyPreprocessing(FRDGBuilder& GraphBuil
 				if (ensureMsgf(ExternalRT != nullptr, TEXT("Unexpected missing external render target override as output.")))
 				{
 					ExternalRDG = GraphBuilder.RegisterExternalTexture(ExternalRT->RenderTarget);
+					GraphBuilder.UseInternalAccessMode(ExternalRDG);
+
 					PassInputs.OverrideOutput = FScreenPassRenderTarget(ExternalRDG, ERenderTargetLoadAction::ENoAction);
 
 					// Update view rect to match output
@@ -483,13 +484,10 @@ void FCompositeCoreSceneViewExtension::ApplyPreprocessing(FRDGBuilder& GraphBuil
 
 			// Note: we do not cache the output to CachedOutputsPerView_RenderThread since we only support cached outputs in post-processing for now.
 
-			// Manually transition the external render target back to a shader resource.
-			if (ExternalRDG != nullptr)
+			if (ExternalRDG)
 			{
-				FCompositeRenderTargetTransitionParameters* Parameters = GraphBuilder.AllocParameters<FCompositeRenderTargetTransitionParameters>();
-				Parameters->TransitionTexture = ExternalRDG;
-
-				GraphBuilder.AddPass(RDG_EVENT_NAME("CompositeCoreTransitionRT"), Parameters, ERDGPassFlags::Raster | ERDGPassFlags::NeverCull | ERDGPassFlags::SkipRenderPass, [](FRHICommandList&) {});
+				// Transition the render target back to external access as a shader resource.
+				GraphBuilder.UseExternalAccessMode(ExternalRDG, ERHIAccess::SRVMask);
 			}
 
 			// Update next pass input
