Composite: Minimal masking blend mode for 2D-only workflows.

diff --git a/Engine/Plugins/Experimental/Compositing/CompositeCore/Shaders/Private/CompositeCoreMerge.usf b/Engine/Plugins/Experimental/Compositing/CompositeCore/Shaders/Private/CompositeCoreMerge.usf
--- a/Engine/Plugins/Experimental/Compositing/CompositeCore/Shaders/Private/CompositeCoreMerge.usf
+++ b/Engine/Plugins/Experimental/Compositing/CompositeCore/Shaders/Private/CompositeCoreMerge.usf
@@ -17,23 +17,26 @@
 
 #define ECompositeCoreMergeOp_None        0
 #define ECompositeCoreMergeOp_Over        1
-#define ECompositeCoreMergeOp_Add         2
-#define ECompositeCoreMergeOp_Subtract    3
-#define ECompositeCoreMergeOp_Mul         4
-#define ECompositeCoreMergeOp_Div         5
-#define ECompositeCoreMergeOp_Min         6
-#define ECompositeCoreMergeOp_Max         7
+#define ECompositeCoreMergeOp_Under       2
+#define ECompositeCoreMergeOp_Add         3
+#define ECompositeCoreMergeOp_Subtract    4
+#define ECompositeCoreMergeOp_Mul         5
+#define ECompositeCoreMergeOp_Div         6
+#define ECompositeCoreMergeOp_Min         7
+#define ECompositeCoreMergeOp_Max         8
+#define ECompositeCoreMergeOp_In          9
+#define ECompositeCoreMergeOp_Mask        10
 // Advanced blending modes
-#define ECompositeCoreMergeOp_Screen      8
-#define ECompositeCoreMergeOp_Overlay     9
-#define ECompositeCoreMergeOp_Darken      10
-#define ECompositeCoreMergeOp_Lighten     11
-#define ECompositeCoreMergeOp_ColorDodge  12
-#define ECompositeCoreMergeOp_ColorBurn   13
-#define ECompositeCoreMergeOp_HardLight   14
-#define ECompositeCoreMergeOp_SoftLight   15
-#define ECompositeCoreMergeOp_Difference  16
-#define ECompositeCoreMergeOp_Exclusion   17
+#define ECompositeCoreMergeOp_Screen      11
+#define ECompositeCoreMergeOp_Overlay     12
+#define ECompositeCoreMergeOp_Darken      13
+#define ECompositeCoreMergeOp_Lighten     14
+#define ECompositeCoreMergeOp_ColorDodge  15
+#define ECompositeCoreMergeOp_ColorBurn   16
+#define ECompositeCoreMergeOp_HardLight   17
+#define ECompositeCoreMergeOp_SoftLight   18
+#define ECompositeCoreMergeOp_Difference  19
+#define ECompositeCoreMergeOp_Exclusion   20
 
 SCREEN_PASS_TEXTURE_VIEWPORT(Input0)
 SCREEN_PASS_TEXTURE_VIEWPORT(Input1)
@@ -167,6 +170,10 @@ void MainPS(
     {
         OutColor = Color0 + Color1 * (1.0 - Color0.a);
     }
+	else if (MergeOp == ECompositeCoreMergeOp_Under)
+	{
+		OutColor = Color0 * (1.0 - Color1.a) + Color1;
+	}
     else if(MergeOp == ECompositeCoreMergeOp_Add)
     {
         OutColor = Color0 + Color1;
@@ -191,6 +198,14 @@ void MainPS(
     {
         OutColor = max(Color0, Color1);
     }
+	else if (MergeOp == ECompositeCoreMergeOp_In)
+	{
+		OutColor = Color0 * Color1.a;
+	}
+	else if (MergeOp == ECompositeCoreMergeOp_Mask)
+	{
+		OutColor = Color1 * Color0.a;
+	}
     //~ Start advanced blending modes
     else if(MergeOp == ECompositeCoreMergeOp_Screen)
     {
diff --git a/Engine/Plugins/Experimental/Compositing/CompositeCore/Source/CompositeCore/Public/Passes/CompositeCorePassMergeProxy.h b/Engine/Plugins/Experimental/Compositing/CompositeCore/Source/CompositeCore/Public/Passes/CompositeCorePassMergeProxy.h
--- a/Engine/Plugins/Experimental/Compositing/CompositeCore/Source/CompositeCore/Public/Passes/CompositeCorePassMergeProxy.h
+++ b/Engine/Plugins/Experimental/Compositing/CompositeCore/Source/CompositeCore/Public/Passes/CompositeCorePassMergeProxy.h
@@ -9,27 +9,33 @@
 
 #define UE_API COMPOSITECORE_API
 
+/**
+* Merge/blend operations, assuming inputs are alpha pre-multiplied.
+*/
 UENUM()
 enum class ECompositeCoreMergeOp : uint8
 {
-	None,
-	Over,
-	Add,
-	Subtract,
-	Multiply,
-	Divide,
-	Min,
-	Max,
-	Screen,
-	Overlay,
-	Darken,
-	Lighten,
-	ColorDodge,
-	ColorBurn,
-	HardLight,
-	SoftLight,
-	Difference,
-	Exclusion,
+	None UMETA(ToolTip = "Current layer A previous replaces previous layer B."),
+	Over UMETA(ToolTip = "Current layer A over previous layer B: A + B * (1-a)."),
+	Under UMETA(ToolTip = "Current layer A under previous layer B: A * (1-b) + B."),
+	Add UMETA(ToolTip = "Current layer A added to previous layer B: A + B."),
+	Subtract UMETA(ToolTip = "Previous layer B subtracted from current layer A: A - B."),
+	Multiply UMETA(ToolTip = "Current layer A multiplied by previous layer B: A * B."),
+	Divide UMETA(ToolTip = "Current layer A (safely) divided by previous layer B: A / B."),
+	Min UMETA(ToolTip = "Per-component minimum between current layer A and previous layer B."),
+	Max UMETA(ToolTip = "Per-component maximum between current layer A and previous layer B."),
+	In UMETA(ToolTip = "Current layer A masked by the previous layer B's alpha: A * b."),
+	Mask UMETA(ToolTip = "Current layer A's alpha masking the previous layer B: B * a."),
+	Screen UMETA(ToolTip = "Screen advanced blend mode."),
+	Overlay UMETA(ToolTip = "Overlay advanced blend mode."),
+	Darken UMETA(ToolTip = "Darken advanced blend mode."),
+	Lighten UMETA(ToolTip = "Lighten advanced blend mode."),
+	ColorDodge UMETA(ToolTip = "Color dodge advanced blend mode."),
+	ColorBurn UMETA(ToolTip = "Color burn advanced blend mode."),
+	HardLight UMETA(ToolTip = "Hard light advanced blend mode."),
+	SoftLight UMETA(ToolTip = "Soft light advanced blend mode."),
+	Difference UMETA(ToolTip = "Difference advanced blend mode."),
+	Exclusion UMETA(ToolTip = "Exclusion advanced blend mode."),
 };
 
 namespace UE
