Composite: Fix last pass identification logic in post-processing.

Recently added processing layers are the exception that break rule where all layers were merge passes. These merge passes are the only type currently which are  "context-aware": automatically managing things like color encodings, alpha inversion & lens distortion. Since then if users added a processing OCIO layer as the last pass, they would lose this automatic management (i.e. the alpha would not be re-inverted correctly). We fix this here by introducing a type name identification on each pass implementation. We then find the last merge pass, so that it can correctly output pixels back into post-processing's scene color.

diff --git a/Engine/Plugins/Experimental/Compositing/Composite/Source/Composite/Private/Layers/CompositeLayerSingleLightShadow.cpp b/Engine/Plugins/Experimental/Compositing/Composite/Source/Composite/Private/Layers/CompositeLayerSingleLightShadow.cpp
--- a/Engine/Plugins/Experimental/Compositing/Composite/Source/Composite/Private/Layers/CompositeLayerSingleLightShadow.cpp
+++ b/Engine/Plugins/Experimental/Compositing/Composite/Source/Composite/Private/Layers/CompositeLayerSingleLightShadow.cpp
@@ -80,6 +80,8 @@ namespace UE
 		class FCompositeSingleLightShadowProxy : public FCompositeCorePassProxy
 		{
 		public:
+			IMPLEMENT_COMPOSITE_PASS(FCompositeSingleLightShadowProxy);
+
 			using FCompositeCorePassProxy::FCompositeCorePassProxy;
 
 			FPassTexture Add(FRDGBuilder& GraphBuilder, const FSceneView& InView, const FPassInputArray& Inputs, const FPassContext& PassContext) const override
diff --git a/Engine/Plugins/Experimental/Compositing/Composite/Source/Composite/Private/Passes/CompositePassCenteredScale.cpp b/Engine/Plugins/Experimental/Compositing/Composite/Source/Composite/Private/Passes/CompositePassCenteredScale.cpp
--- a/Engine/Plugins/Experimental/Compositing/Composite/Source/Composite/Private/Passes/CompositePassCenteredScale.cpp
+++ b/Engine/Plugins/Experimental/Compositing/Composite/Source/Composite/Private/Passes/CompositePassCenteredScale.cpp
@@ -33,6 +33,8 @@ namespace UE
 		class FCenteredScalePassProxy : public FCompositeCorePassProxy
 		{
 		public:
+			IMPLEMENT_COMPOSITE_PASS(FCenteredScalePassProxy);
+
 			using FCompositeCorePassProxy::FCompositeCorePassProxy;
 
 			FPassTexture Add(FRDGBuilder& GraphBuilder, const FSceneView& InView, const FPassInputArray& Inputs, const FPassContext& PassContext) const override
diff --git a/Engine/Plugins/Experimental/Compositing/Composite/Source/Composite/Private/Passes/CompositePassColorGrade.cpp b/Engine/Plugins/Experimental/Compositing/Composite/Source/Composite/Private/Passes/CompositePassColorGrade.cpp
--- a/Engine/Plugins/Experimental/Compositing/Composite/Source/Composite/Private/Passes/CompositePassColorGrade.cpp
+++ b/Engine/Plugins/Experimental/Compositing/Composite/Source/Composite/Private/Passes/CompositePassColorGrade.cpp
@@ -56,6 +56,8 @@ IMPLEMENT_GLOBAL_SHADER(FCompositeColorGradeShader, "/Plugin/Composite/Private/C
 class FColorGradeCompositePassProxy : public FCompositeCorePassProxy
 {
 public:
+	IMPLEMENT_COMPOSITE_PASS(FColorGradeCompositePassProxy);
+
 	FColorGradeCompositePassProxy(UE::CompositeCore::FPassInputDeclArray InPassDeclaredInputs)
 		: FCompositeCorePassProxy(MoveTemp(InPassDeclaredInputs))
 	{ }
diff --git a/Engine/Plugins/Experimental/Compositing/Composite/Source/Composite/Private/Passes/CompositePassDeltaKeyer.cpp b/Engine/Plugins/Experimental/Compositing/Composite/Source/Composite/Private/Passes/CompositePassDeltaKeyer.cpp
--- a/Engine/Plugins/Experimental/Compositing/Composite/Source/Composite/Private/Passes/CompositePassDeltaKeyer.cpp
+++ b/Engine/Plugins/Experimental/Compositing/Composite/Source/Composite/Private/Passes/CompositePassDeltaKeyer.cpp
@@ -81,6 +81,8 @@ namespace UE
 		class FCompositePassDeltaKeyerProxy : public FCompositeCorePassProxy
 		{
 		public:
+			IMPLEMENT_COMPOSITE_PASS(FCompositePassDeltaKeyerProxy);
+
 			using FCompositeCorePassProxy::FCompositeCorePassProxy;
 
 			FScreenPassTexture AddDenoisePass(FRDGBuilder& GraphBuilder, const FScreenPassTexture& Input, ERHIFeatureLevel::Type FeatureLevel) const
diff --git a/Engine/Plugins/Experimental/Compositing/Composite/Source/Composite/Private/Passes/CompositePassDistortion.cpp b/Engine/Plugins/Experimental/Compositing/Composite/Source/Composite/Private/Passes/CompositePassDistortion.cpp
--- a/Engine/Plugins/Experimental/Compositing/Composite/Source/Composite/Private/Passes/CompositePassDistortion.cpp
+++ b/Engine/Plugins/Experimental/Compositing/Composite/Source/Composite/Private/Passes/CompositePassDistortion.cpp
@@ -39,6 +39,8 @@ DECLARE_GPU_STAT_NAMED(FCompositeDistortion, TEXT("Composite.Distortion"));
 class FDistortionCompositePassProxy : public FCompositeCorePassProxy
 {
 public:
+	IMPLEMENT_COMPOSITE_PASS(FDistortionCompositePassProxy);
+
 	enum class EDistortionUV
 	{
 		None = 0,
diff --git a/Engine/Plugins/Experimental/Compositing/Composite/Source/Composite/Private/Passes/CompositePassMaterial.cpp b/Engine/Plugins/Experimental/Compositing/Composite/Source/Composite/Private/Passes/CompositePassMaterial.cpp
--- a/Engine/Plugins/Experimental/Compositing/Composite/Source/Composite/Private/Passes/CompositePassMaterial.cpp
+++ b/Engine/Plugins/Experimental/Compositing/Composite/Source/Composite/Private/Passes/CompositePassMaterial.cpp
@@ -14,6 +14,8 @@ DECLARE_GPU_STAT_NAMED(FCompositeMaterial, TEXT("Composite.Material"));
 class FCompositePassMaterialProxy : public FCompositeCorePassProxy
 {
 public:
+	IMPLEMENT_COMPOSITE_PASS(FCompositePassMaterialProxy);
+
 	using FCompositeCorePassProxy::FCompositeCorePassProxy;
 
 	/* Post-process material weak pointer. */
diff --git a/Engine/Plugins/Experimental/Compositing/Composite/Source/Composite/Private/Passes/CompositePassOpenColorIO.cpp b/Engine/Plugins/Experimental/Compositing/Composite/Source/Composite/Private/Passes/CompositePassOpenColorIO.cpp
--- a/Engine/Plugins/Experimental/Compositing/Composite/Source/Composite/Private/Passes/CompositePassOpenColorIO.cpp
+++ b/Engine/Plugins/Experimental/Compositing/Composite/Source/Composite/Private/Passes/CompositePassOpenColorIO.cpp
@@ -13,6 +13,8 @@ class FOpenColorIOCompositePassProxy
 	, public TSharedFromThis<FOpenColorIOCompositePassProxy, ESPMode::ThreadSafe>
 {
 public:
+	IMPLEMENT_COMPOSITE_PASS(FOpenColorIOCompositePassProxy);
+
 	FOpenColorIOCompositePassProxy(UE::CompositeCore::FPassInputDeclArray InPassDeclaredInputs, FOpenColorIORenderPassResources&& InPassResources)
 		: FCompositeCorePassProxy(MoveTemp(InPassDeclaredInputs))
 		, PassResources(InPassResources)
diff --git a/Engine/Plugins/Experimental/Compositing/Composite/Source/Composite/Private/Passes/CompositePassSMAA.cpp b/Engine/Plugins/Experimental/Compositing/Composite/Source/Composite/Private/Passes/CompositePassSMAA.cpp
--- a/Engine/Plugins/Experimental/Compositing/Composite/Source/Composite/Private/Passes/CompositePassSMAA.cpp
+++ b/Engine/Plugins/Experimental/Compositing/Composite/Source/Composite/Private/Passes/CompositePassSMAA.cpp
@@ -20,6 +20,8 @@ namespace UE
 		class FSMAAPassProxy : public FCompositeCorePassProxy
 		{
 		public:
+			IMPLEMENT_COMPOSITE_PASS(FSMAAPassProxy);
+
 			using FCompositeCorePassProxy::FCompositeCorePassProxy;
 
 			FPassTexture Add(FRDGBuilder& GraphBuilder, const FSceneView& InView, const FPassInputArray& Inputs, const FPassContext& PassContext) const override
diff --git a/Engine/Plugins/Experimental/Compositing/CompositeCore/Source/CompositeCore/Private/CompositeCoreSceneViewExtension.cpp b/Engine/Plugins/Experimental/Compositing/CompositeCore/Source/CompositeCore/Private/CompositeCoreSceneViewExtension.cpp
--- a/Engine/Plugins/Experimental/Compositing/CompositeCore/Source/CompositeCore/Private/CompositeCoreSceneViewExtension.cpp
+++ b/Engine/Plugins/Experimental/Compositing/CompositeCore/Source/CompositeCore/Private/CompositeCoreSceneViewExtension.cpp
@@ -5,6 +5,7 @@
 #include "CompositeCoreModule.h"
 #include "Passes/CompositeCorePassDilate.h"
 #include "Passes/CompositeCorePassFXAAProxy.h"
+#include "Passes/CompositeCorePassMergeProxy.h"
 
 #include "CommonRenderResources.h"
 #include "Components/PrimitiveComponent.h"
@@ -534,7 +535,8 @@ bool FCompositeCoreSceneViewExtension::ApplyPasses_Recursive(
 	UE::CompositeCore::FPassContext& PassContext,
 	const TArray<const FCompositeCorePassProxy*> InPasses,
 	FExternalTextureMap& ExternalTextures,
-	int32 RecursionLevel,
+	const int32 LastMergePassIndex,
+	const int32 RecursionLevel,
 	UE::CompositeCore::FPassTexture& Output)
 {
 	using namespace UE::CompositeCore;
@@ -553,8 +555,8 @@ bool FCompositeCoreSceneViewExtension::ApplyPasses_Recursive(
 	{
 		const FCompositeCorePassProxy* Pass = InPasses[PassIndex];
 		
-		// Output override only applies to the (base) last pass
-		const bool bIsLastPass = (PassIndex == InPasses.Num() - 1) && (RecursionLevel == 0);
+		// Output override only applies to the last merge pass
+		const bool bIsLastMergePass = (RecursionLevel == 0) && (PassIndex == LastMergePassIndex);
 
 		// Update inputs for the current pass
 		FPassInputArray PassInputs = BasePassInputs;
@@ -575,7 +577,7 @@ bool FCompositeCoreSceneViewExtension::ApplyPasses_Recursive(
 				{
 					PassInputs[InputIndex] = *PreProducedOutput;
 				}
-				else if(ApplyPasses_Recursive(GraphBuilder, InView, PassInputs, OriginalInputs, PassContext, { ChildPassProxy }, ExternalTextures, RecursionLevel + 1, ChildPassOutput))
+				else if(ApplyPasses_Recursive(GraphBuilder, InView, PassInputs, OriginalInputs, PassContext, { ChildPassProxy }, ExternalTextures, INDEX_NONE, RecursionLevel + 1, ChildPassOutput))
 				{
 					PassInputs[InputIndex] = ChildPassOutput;
 				}
@@ -619,7 +621,7 @@ bool FCompositeCoreSceneViewExtension::ApplyPasses_Recursive(
 			}
 		}
 
-		if (bIsLastPass)
+		if (bIsLastMergePass)
 		{
 			// Invert alpha when writing back to scene color
 			PassContext.bOutputSceneColor = true;
@@ -657,9 +659,14 @@ FScreenPassTexture FCompositeCoreSceneViewExtension::PostProcessWork_RenderThrea
 
 	FPassInputArray ResolvedInputs(GraphBuilder, InView, Inputs, InLocation);
 	FPassTexture Output;
-	
+
+	const int32 LastMergePassIndex = Passes.FindLastByPredicate([](const FCompositeCorePassProxy* InPass)
+		{
+			return InPass ? (InPass->GetTypeName() == GetMergePassPassName()) : false;
+		});
 	constexpr int32 RecursionLevel = 0;
-	if (ApplyPasses_Recursive(GraphBuilder, InView, ResolvedInputs, ResolvedInputs, PassContext, Passes, ExternalTextures, RecursionLevel, Output))
+	
+	if (ApplyPasses_Recursive(GraphBuilder, InView, ResolvedInputs, ResolvedInputs, PassContext, Passes, ExternalTextures, LastMergePassIndex, RecursionLevel, Output))
 	{
 		return MoveTemp(Output.Texture);
 	}
diff --git a/Engine/Plugins/Experimental/Compositing/CompositeCore/Source/CompositeCore/Private/CompositeCoreSceneViewExtension.h b/Engine/Plugins/Experimental/Compositing/CompositeCore/Source/CompositeCore/Private/CompositeCoreSceneViewExtension.h
--- a/Engine/Plugins/Experimental/Compositing/CompositeCore/Source/CompositeCore/Private/CompositeCoreSceneViewExtension.h
+++ b/Engine/Plugins/Experimental/Compositing/CompositeCore/Source/CompositeCore/Private/CompositeCoreSceneViewExtension.h
@@ -102,7 +102,8 @@ private:
 		UE::CompositeCore::FPassContext& PassContext,
 		const TArray<const FCompositeCorePassProxy*> InPasses,
 		FExternalTextureMap& InOutExternalTextures,
-		int32 RecursionLevel,
+		const int32 LastMergePassIndex,
+		const int32 RecursionLevel,
 		UE::CompositeCore::FPassTexture& Output
 	);
 
diff --git a/Engine/Plugins/Experimental/Compositing/CompositeCore/Source/CompositeCore/Private/Passes/CompositeCorePassMergeProxy.cpp b/Engine/Plugins/Experimental/Compositing/CompositeCore/Source/CompositeCore/Private/Passes/CompositeCorePassMergeProxy.cpp
--- a/Engine/Plugins/Experimental/Compositing/CompositeCore/Source/CompositeCore/Private/Passes/CompositeCorePassMergeProxy.cpp
+++ b/Engine/Plugins/Experimental/Compositing/CompositeCore/Source/CompositeCore/Private/Passes/CompositeCorePassMergeProxy.cpp
@@ -226,6 +226,11 @@ namespace UE
 
 			return FPassTexture{ MoveTemp(Output), MoveTemp(OutMetadata) };
 		}
-	}
+
+		const FName& GetMergePassPassName()
+		{
+			return FMergePassProxy::GetTypeNameStatic();
+		}
+}
 }
 
diff --git a/Engine/Plugins/Experimental/Compositing/CompositeCore/Source/CompositeCore/Public/Passes/CompositeCorePassFXAAProxy.h b/Engine/Plugins/Experimental/Compositing/CompositeCore/Source/CompositeCore/Public/Passes/CompositeCorePassFXAAProxy.h
--- a/Engine/Plugins/Experimental/Compositing/CompositeCore/Source/CompositeCore/Public/Passes/CompositeCorePassFXAAProxy.h
+++ b/Engine/Plugins/Experimental/Compositing/CompositeCore/Source/CompositeCore/Public/Passes/CompositeCorePassFXAAProxy.h
@@ -19,6 +19,8 @@ namespace UE
 		class FFXAAPassProxy : public FCompositeCorePassProxy
 		{
 		public:
+			IMPLEMENT_COMPOSITE_PASS(FFXAAPassProxy);
+
 			using FCompositeCorePassProxy::FCompositeCorePassProxy;
 
 			UE_API FPassTexture Add(FRDGBuilder& GraphBuilder, const FSceneView& InView, const FPassInputArray& Inputs, const FPassContext& PassContext) const override;
diff --git a/Engine/Plugins/Experimental/Compositing/CompositeCore/Source/CompositeCore/Public/Passes/CompositeCorePassMergeProxy.h b/Engine/Plugins/Experimental/Compositing/CompositeCore/Source/CompositeCore/Public/Passes/CompositeCorePassMergeProxy.h
--- a/Engine/Plugins/Experimental/Compositing/CompositeCore/Source/CompositeCore/Public/Passes/CompositeCorePassMergeProxy.h
+++ b/Engine/Plugins/Experimental/Compositing/CompositeCore/Source/CompositeCore/Public/Passes/CompositeCorePassMergeProxy.h
@@ -52,6 +52,8 @@ namespace UE
 		class FMergePassProxy : public FCompositeCorePassProxy
 		{
 		public:
+			IMPLEMENT_COMPOSITE_PASS(FMergePassProxy);
+
 			using FCompositeCorePassProxy::FCompositeCorePassProxy;
 
 			/** Constructor */
@@ -59,6 +61,7 @@ namespace UE
 
 			UE_API FPassTexture Add(FRDGBuilder& GraphBuilder, const FSceneView& InView, const FPassInputArray& Inputs, const FPassContext& PassContext) const override;
 
+		public:
 			/** Merge operation. */
 			ECompositeCoreMergeOp MergeOp = ECompositeCoreMergeOp::Over;
 
@@ -68,6 +71,8 @@ namespace UE
 			/** Flag to enable or disable application of engine-provided lens distortion LUTs, true by default. */
 			ELensDistortionHandling bUseLensDistortion = ELensDistortionHandling::Enabled;
 		};
+
+		const FName& GetMergePassPassName();
 	}
 }
 
diff --git a/Engine/Plugins/Experimental/Compositing/CompositeCore/Source/CompositeCore/Public/Passes/CompositeCorePassProxy.h b/Engine/Plugins/Experimental/Compositing/CompositeCore/Source/CompositeCore/Public/Passes/CompositeCorePassProxy.h
--- a/Engine/Plugins/Experimental/Compositing/CompositeCore/Source/CompositeCore/Public/Passes/CompositeCorePassProxy.h
+++ b/Engine/Plugins/Experimental/Compositing/CompositeCore/Source/CompositeCore/Public/Passes/CompositeCorePassProxy.h
@@ -246,6 +246,9 @@ public:
 	/** Virtual destructor */
 	virtual ~FCompositeCorePassProxy() = default;
 
+	/** Bare-bones RTTI method to allow users to differentiate / downcast the composite core pass they're interested in. */
+	virtual const FName& GetTypeName() const PURE_VIRTUAL(FCompositeCorePassProxy::GetTypeName, static FName Name; return Name;);
+
 	/** Render-thread add pass method to override. */
 	virtual UE::CompositeCore::FPassTexture Add(
 		FRDGBuilder& GraphBuilder,
@@ -306,4 +309,9 @@ protected:
 	TOptional<UE::CompositeCore::ResourceId> PassDeclaredOutputOverride;
 };
 
+/** Use this macro to implement the RTTI method for composite render passes. Should only be used on final classes (e.g. if C derives from B which derives from A, we won't be able
+  *  to query a C* for whether it's a "B", in other words, A and B should be abstract and C, final) */
+#define IMPLEMENT_COMPOSITE_PASS(TypeName) const FName& GetTypeName() const override { return GetTypeNameStatic(); } \
+	static const FName& GetTypeNameStatic() { static FName Name(TEXT(#TypeName)); return Name; }
+
 #undef UE_API
