Composite: Fix view family render target final write override, which broke at CL 45564601.

* The last merge pass will output scene-color compatible values
* The last pass (which could differ from the last merge pass) will output to the viewport render target

diff --git a/Engine/Plugins/Compositing/CompositeCore/Source/CompositeCore/Private/CompositeCoreSceneViewExtension.cpp b/Engine/Plugins/Compositing/CompositeCore/Source/CompositeCore/Private/CompositeCoreSceneViewExtension.cpp
--- a/Engine/Plugins/Compositing/CompositeCore/Source/CompositeCore/Private/CompositeCoreSceneViewExtension.cpp
+++ b/Engine/Plugins/Compositing/CompositeCore/Source/CompositeCore/Private/CompositeCoreSceneViewExtension.cpp
@@ -555,8 +555,16 @@ bool FCompositeCoreSceneViewExtension::ApplyPasses_Recursive(
 	{
 		const FCompositeCorePassProxy* Pass = InPasses[PassIndex];
 		
-		// Output override only applies to the last merge pass
-		const bool bIsLastMergePass = (RecursionLevel == 0) && (PassIndex == LastMergePassIndex);
+		// Last merge pass is used to output values compatible with scene color
+		bool bIsLastMergePass = false;
+		// Last pass writes to the viewport render target
+		bool bIsLastPass = false;
+		
+		if (RecursionLevel == 0)
+		{
+			bIsLastMergePass = (PassIndex == LastMergePassIndex);
+			bIsLastPass = (PassIndex == InPasses.Num() - 1);
+		}
 
 		// Update inputs for the current pass
 		FPassInputArray PassInputs = BasePassInputs;
@@ -623,10 +631,11 @@ bool FCompositeCoreSceneViewExtension::ApplyPasses_Recursive(
 
 		if (bIsLastMergePass)
 		{
-			// Invert alpha when writing back to scene color
+			// Merge pass proxies will handle alpha inversion, lens distortion and color encodings when writing back to scene color
 			PassContext.bOutputSceneColor = true;
 		}
-		else
+
+		if (!bIsLastPass)
 		{
 			// Only apply the output override on the last pass
 			PassInputs.OverrideOutput = FScreenPassRenderTarget{};
