CompositeCore: Fix bughawk-reported issue to allow multiple view processing, instead of an early return on the first view.

#rnx

diff --git a/Engine/Plugins/Compositing/CompositeCore/Source/CompositeCore/Private/CompositeCoreSceneViewExtension.cpp b/Engine/Plugins/Compositing/CompositeCore/Source/CompositeCore/Private/CompositeCoreSceneViewExtension.cpp
--- a/Engine/Plugins/Compositing/CompositeCore/Source/CompositeCore/Private/CompositeCoreSceneViewExtension.cpp
+++ b/Engine/Plugins/Compositing/CompositeCore/Source/CompositeCore/Private/CompositeCoreSceneViewExtension.cpp
@@ -302,7 +302,7 @@ void FCompositeCoreSceneViewExtension::BeginRenderViewFamily(FSceneViewFamily& I
 
 		if (CompositeCorePrimitiveIds.IsEmpty())
 		{
-			return;
+			continue;
 		}
 
 		UE::CompositeCore::FBuiltInRenderPassOptions RenderPassOptions = BuiltInRenderPassOptions.IsSet() ? *BuiltInRenderPassOptions : UE::CompositeCore::FBuiltInRenderPassOptions{};
