Composite: Improvements and fixes to multi-viewport workflows.

* Options to constrain the composite renders (including scene captures) to viewports with lit view modes, or the media profile.
* Manually updating scene captures in the scene view extension instead of relying on capture every frame.
* Remove layer scene capture management of bCaptureEveryFrame: instead we only register scene captures to update when active.
* Remove CompositeCore.RestrictedHoldoutViewMode cvar which was a temporary solution to this problem.

Note that further work is still needed to prevent visual conflicts between scene captures updating the same render targets across multiple lit viewports.

diff --git a/Engine/Plugins/Compositing/Composite/Source/Composite/Private/Components/CompositeSceneCapture2DComponent.cpp b/Engine/Plugins/Compositing/Composite/Source/Composite/Private/Components/CompositeSceneCapture2DComponent.cpp
--- a/Engine/Plugins/Compositing/Composite/Source/Composite/Private/Components/CompositeSceneCapture2DComponent.cpp
+++ b/Engine/Plugins/Compositing/Composite/Source/Composite/Private/Components/CompositeSceneCapture2DComponent.cpp
@@ -15,6 +15,11 @@ UCompositeSceneCapture2DComponent::UCompositeSceneCapture2DComponent(const FObje
 {
 	PrimitiveRenderMode = ESceneCapturePrimitiveRenderMode::PRM_UseShowOnlyList;
 
+	// We issue scene capture render updates manually in FCompositeCoreSceneViewExtension
+	bCaptureEveryFrame = false;
+	bAlwaysPersistRenderingState = true;
+	bCaptureOnMovement = false;
+
 	CaptureSource = ESceneCaptureSource::SCS_FinalColorHDR;
 	bUseRayTracingIfEnabled = true;
 	bMainViewFamily = true;
diff --git a/Engine/Plugins/Compositing/Composite/Source/Composite/Private/CompositeActor.cpp b/Engine/Plugins/Compositing/Composite/Source/Composite/Private/CompositeActor.cpp
--- a/Engine/Plugins/Compositing/Composite/Source/Composite/Private/CompositeActor.cpp
+++ b/Engine/Plugins/Compositing/Composite/Source/Composite/Private/CompositeActor.cpp
@@ -52,6 +52,22 @@ namespace UE
 
 		namespace Private
 		{
+			/** Convert the allowed viewport view mode into a set of view mode indices. */
+			inline TSet<EViewModeIndex> GetViewModeIndices(ECompositeAllowedViewport InMode)
+			{
+				switch (InMode)
+				{
+				case ECompositeAllowedViewport::Default:
+					return { VMI_Lit, VMI_PathTracing, VMI_Unknown };
+				case ECompositeAllowedViewport::MediaProfileOnly:
+					return { VMI_Unknown };
+				default:
+					break;
+				}
+
+				return {};
+			}
+
 			/** Return solo layer from list. */
 			TObjectPtr<UCompositeLayerBase> FindSoloLayer(TArrayView<const TObjectPtr<UCompositeLayerBase>> InLayers)
 			{
@@ -120,7 +136,9 @@ namespace UE
 			}
 
 			/* Gather render work (per frame) given the current composite actor properties. */
-			UE::CompositeCore::FRenderWork GatherFrameRenderWork(const ACompositeActor& InCompositeActor)
+			UE::CompositeCore::FRenderWork GatherFrameRenderWork(
+				const ACompositeActor& InCompositeActor,
+				const TMap<TWeakObjectPtr<const UCompositeLayerBase>, FSceneCaptureComponentArray>& InSceneCapturesPerLayer)
 			{
 				if (InCompositeActor.CompositeLayers.IsEmpty())
 				{
@@ -191,6 +209,23 @@ namespace UE
 				default:
 					break;
 				}
+				
+				RenderWork.AllowedViewModes = GetViewModeIndices(InCompositeActor.AllowedViewport);
+
+				// Currently, each layer may use at most 2 scene captures, so we reserve the maximum case.
+				RenderWork.SceneCapturesUpdateQueue.Reserve(InSceneCapturesPerLayer.Num() * 2);
+
+				for (const auto& Pair : InSceneCapturesPerLayer)
+				{
+					const TWeakObjectPtr<const UCompositeLayerBase>& Layer = Pair.Key;
+					if (Layer.IsValid() && Layer->IsEnabled())
+					{
+						for (const TObjectPtr<UCompositeSceneCapture2DComponent>& SceneCapture : Pair.Value.Components)
+						{
+							RenderWork.SceneCapturesUpdateQueue.Add(SceneCapture.Get());
+						}
+					}
+				}
 
 				return RenderWork;
 			}
@@ -210,6 +245,8 @@ ACompositeActor::ACompositeActor(const FObjectInitializer& ObjectInitializer)
 	, bEnableScreenSpaceReflections{true}
 	, bOverridesViewUserFlags{true}
 	, ViewUserFlags{1}
+	, MainRenderOutput{ECompositeMainRenderOutputMode::Default}
+	, AllowedViewport{ECompositeAllowedViewport::Default}
 	, bIsModifyingAProperty{false}
 {
 	PrimaryActorTick.bCanEverTick = true;
@@ -334,7 +371,7 @@ void ACompositeActor::Tick(float DeltaTime)
 	UCompositeCoreSubsystem* Subsystem = UWorld::GetSubsystem<UCompositeCoreSubsystem>(GetWorld());
 	if (IsValid(Subsystem))
 	{
-		Subsystem->SetRenderWork(UE::Composite::Private::GatherFrameRenderWork(*this));
+		Subsystem->SetRenderWork(UE::Composite::Private::GatherFrameRenderWork(*this, SceneCapturesPerLayer));
 
 		UE::CompositeCore::FBuiltInRenderPassOptions Options;
 		if (bOverridesViewUserFlags)
@@ -676,20 +713,6 @@ UCompositeShadowReflectionCatcherComponent* ACompositeActor::FindOrCreateSceneCa
 	return Component;
 }
 
-void ACompositeActor::ActivateSceneCaptures(const UCompositeLayerBase* InLayer, bool bInCaptureEveryFrame)
-{
-	if (FSceneCaptureComponentArray* SceneCapturesPtr = SceneCapturesPerLayer.Find(InLayer))
-	{
-		for (TObjectPtr<UCompositeSceneCapture2DComponent>& SceneCapture : SceneCapturesPtr->Components)
-		{
-			if (IsValid(SceneCapture))
-			{
-				SceneCapture->bCaptureEveryFrame = bInCaptureEveryFrame;
-			}
-		}
-	}
-}
-
 void ACompositeActor::DestroySceneCaptures(const UCompositeLayerBase* InLayer)
 {
 	if (!ensureMsgf(InLayer, TEXT("Expected non-null composite layer.")))
@@ -749,11 +772,6 @@ void ACompositeActor::ApplyAllState()
 	UCompositeCoreSubsystem* Subsystem = UWorld::GetSubsystem<UCompositeCoreSubsystem>(GetWorld());
 
 	OnEnabledStateChange.ExecuteIfBound(true, Subsystem);
-
-	for (const TObjectPtr<UCompositeLayerBase>& CompositeLayer : CompositeLayers)
-	{
-		ActivateSceneCaptures(CompositeLayer.Get(), true);
-	}
 }
 
 void ACompositeActor::RemoveAllState()
@@ -767,11 +785,6 @@ void ACompositeActor::RemoveAllState()
 		Subsystem->ResetRenderWork();
 		Subsystem->ResetBuiltInRenderPassOptions();
 	}
-
-	for (const TObjectPtr<UCompositeLayerBase>& CompositeLayer : CompositeLayers)
-	{
-		ActivateSceneCaptures(CompositeLayer.Get(), false);
-	}
 }
 
 
diff --git a/Engine/Plugins/Compositing/Composite/Source/Composite/Private/Layers/CompositeLayerSceneCapture.cpp b/Engine/Plugins/Compositing/Composite/Source/Composite/Private/Layers/CompositeLayerSceneCapture.cpp
--- a/Engine/Plugins/Compositing/Composite/Source/Composite/Private/Layers/CompositeLayerSceneCapture.cpp
+++ b/Engine/Plugins/Compositing/Composite/Source/Composite/Private/Layers/CompositeLayerSceneCapture.cpp
@@ -18,22 +18,7 @@ UCompositeLayerSceneCapture::UCompositeLayerSceneCapture(const FObjectInitialize
 
 UCompositeLayerSceneCapture::~UCompositeLayerSceneCapture() = default;
 
-void UCompositeLayerSceneCapture::SetEnabled(bool bInEnabled)
-{
-	Super::SetEnabled(bInEnabled);
-
-	ACompositeActor* CompositeActor = GetTypedOuter<ACompositeActor>();
-	if (IsValid(CompositeActor))
-	{
-		CompositeActor->ActivateSceneCaptures(this, IsEnabled());
-	}
-}
 
-void UCompositeLayerSceneCapture::BeginDestroy()
-{
-	Super::BeginDestroy();
-}
-	
 #if WITH_EDITOR
 void UCompositeLayerSceneCapture::PostEditChangeProperty(FPropertyChangedEvent& PropertyChangedEvent)
 {
@@ -41,11 +26,7 @@ void UCompositeLayerSceneCapture::PostEditChangeProperty(FPropertyChangedEvent&
 
 	const FName PropertyName = PropertyChangedEvent.GetMemberPropertyName();
 	
-	if (PropertyName == GET_MEMBER_NAME_CHECKED(ThisClass, bIsEnabled))
-	{
-		SetEnabled(bIsEnabled);
-	}
-	else if (PropertyName == GET_MEMBER_NAME_CHECKED(ThisClass, Actors))
+	if (PropertyName == GET_MEMBER_NAME_CHECKED(ThisClass, Actors))
 	{
 		SetActors(MoveTemp(Actors));
 	}
diff --git a/Engine/Plugins/Compositing/Composite/Source/Composite/Private/Layers/CompositeLayerShadowReflection.cpp b/Engine/Plugins/Compositing/Composite/Source/Composite/Private/Layers/CompositeLayerShadowReflection.cpp
--- a/Engine/Plugins/Compositing/Composite/Source/Composite/Private/Layers/CompositeLayerShadowReflection.cpp
+++ b/Engine/Plugins/Compositing/Composite/Source/Composite/Private/Layers/CompositeLayerShadowReflection.cpp
@@ -17,17 +17,6 @@ UCompositeLayerShadowReflection::UCompositeLayerShadowReflection(const FObjectIn
 
 UCompositeLayerShadowReflection::~UCompositeLayerShadowReflection() = default;
 
-void UCompositeLayerShadowReflection::SetEnabled(bool bInEnabled)
-{
-	Super::SetEnabled(bInEnabled);
-
-	ACompositeActor* CompositeActor = GetTypedOuter<ACompositeActor>();
-	if (IsValid(CompositeActor))
-	{
-		CompositeActor->ActivateSceneCaptures(this, IsEnabled());
-	}
-}
-
 void UCompositeLayerShadowReflection::BeginDestroy()
 {
 	Super::BeginDestroy();
@@ -68,11 +57,7 @@ void UCompositeLayerShadowReflection::PostEditChangeProperty(FPropertyChangedEve
 
 	const FName PropertyName = PropertyChangedEvent.GetMemberPropertyName();
 
-	if (PropertyName == GET_MEMBER_NAME_CHECKED(ThisClass, bIsEnabled))
-	{
-		SetEnabled(bIsEnabled);
-	}
-	else if (PropertyName == GET_MEMBER_NAME_CHECKED(ThisClass, Actors))
+	if (PropertyName == GET_MEMBER_NAME_CHECKED(ThisClass, Actors))
 	{
 		SetActors(MoveTemp(Actors));
 	}
diff --git a/Engine/Plugins/Compositing/Composite/Source/Composite/Private/Layers/CompositeLayerSingleLightShadow.cpp b/Engine/Plugins/Compositing/Composite/Source/Composite/Private/Layers/CompositeLayerSingleLightShadow.cpp
--- a/Engine/Plugins/Compositing/Composite/Source/Composite/Private/Layers/CompositeLayerSingleLightShadow.cpp
+++ b/Engine/Plugins/Compositing/Composite/Source/Composite/Private/Layers/CompositeLayerSingleLightShadow.cpp
@@ -164,22 +164,6 @@ UCompositeLayerSingleLightShadow::UCompositeLayerSingleLightShadow(const FObject
 	Operation = ECompositeCoreMergeOp::Multiply;
 }
 
-void UCompositeLayerSingleLightShadow::SetEnabled(bool bInEnabled)
-{
-	Super::SetEnabled(bInEnabled);
-
-	ACompositeActor* CompositeActor = GetTypedOuter<ACompositeActor>();
-	if (IsValid(CompositeActor))
-	{
-		CompositeActor->ActivateSceneCaptures(this, IsEnabled());
-	}
-}
-
-void UCompositeLayerSingleLightShadow::BeginDestroy()
-{
-	Super::BeginDestroy();
-}
-
 #if WITH_EDITOR
 void UCompositeLayerSingleLightShadow::PostEditChangeProperty(FPropertyChangedEvent& PropertyChangedEvent)
 {
@@ -187,15 +171,7 @@ void UCompositeLayerSingleLightShadow::PostEditChangeProperty(FPropertyChangedEv
 
 	const FName PropertyName = PropertyChangedEvent.GetMemberPropertyName();
 
-	if (PropertyName == GET_MEMBER_NAME_CHECKED(ThisClass, bIsEnabled))
-	{
-		ACompositeActor* CompositeActor = GetTypedOuter<ACompositeActor>();
-		if (IsValid(CompositeActor))
-		{
-			CompositeActor->ActivateSceneCaptures(this, bIsEnabled);
-		}
-	}
-	else if (PropertyName == GET_MEMBER_NAME_CHECKED(ThisClass, ShadowCastingActors))
+	if (PropertyName == GET_MEMBER_NAME_CHECKED(ThisClass, ShadowCastingActors))
 	{
 		SetShadowCastingActors(ShadowCastingActors);
 	}
diff --git a/Engine/Plugins/Compositing/Composite/Source/Composite/Public/CompositeActor.h b/Engine/Plugins/Compositing/Composite/Source/Composite/Public/CompositeActor.h
--- a/Engine/Plugins/Compositing/Composite/Source/Composite/Public/CompositeActor.h
+++ b/Engine/Plugins/Compositing/Composite/Source/Composite/Public/CompositeActor.h
@@ -39,6 +39,17 @@ enum class ECompositeMainRenderOutputMode : uint8
 	FinalToneCurveHDR UMETA(DisplayName = "HDR Linear with Tone Curve", ToolTip = "Linear with tone curve output from post-processing's tonemap, equivalent to Final Color (with tone curve) in Linear sRGB gamut."),
 };
 
+/**
+* Constrain composite render to certain view modes for multi-viewport workflows where unlit or wireframe viewports will not conflict with the composite.
+*/
+UENUM(BlueprintType)
+enum class ECompositeAllowedViewport : uint8
+{
+	Default UMETA(ToolTip = "Compositing is allowed on viewports with Lit, Path Tracing or Unknown view modes."),
+	MediaProfileOnly UMETA(ToolTip = "Compositing is only allowed with the Unknown view mode, the default when media profile does its own rendering."),
+	All UMETA(ToolTip = "Compositing is allowed with all view modes."),
+};
+
 /** Actor used to control properties of the composite pipeline. */
 UCLASS(MinimalAPI)
 class ACompositeActor : public AActor
@@ -143,6 +154,10 @@ public:
 	UPROPERTY(EditAnywhere, BlueprintReadWrite, Category = "Composite", AdvancedDisplay)
 	ECompositeMainRenderOutputMode MainRenderOutput;
 
+	/** Constrain composite rendering to specific viewport view modes. */
+	UPROPERTY(EditAnywhere, BlueprintReadWrite, Category = "Composite", AdvancedDisplay)
+	ECompositeAllowedViewport AllowedViewport;
+
 	/** Delegate called when joining multi-user. Can be used to automatically re-active the actor locally. */
 	UPROPERTY(BlueprintAssignable, BlueprintCallable, BlueprintReadWrite, Category = "Composite")
 	FOnCompositeActorJoinMultiUserSession OnJoinMultiUserSession;
@@ -152,9 +167,6 @@ public:
 	template<class RetType>
 	UE_API RetType* FindOrCreateSceneCapture(const UCompositeLayerBase* InLayer, int32 InIndex = 0, FName InBaseName = NAME_None);
 
-	/** Convenience function used by layers to control the capture state of scene captures. */
-	UE_API void ActivateSceneCaptures(const UCompositeLayerBase* InLayer, bool bInCaptureEveryFrame);
-
 	/** Convenience function used by layers to remove and scene capture components. */
 	UE_API void DestroySceneCaptures(const UCompositeLayerBase* InLayer);
 
diff --git a/Engine/Plugins/Compositing/Composite/Source/Composite/Public/Layers/CompositeLayerSceneCapture.h b/Engine/Plugins/Compositing/Composite/Source/Composite/Public/Layers/CompositeLayerSceneCapture.h
--- a/Engine/Plugins/Compositing/Composite/Source/Composite/Public/Layers/CompositeLayerSceneCapture.h
+++ b/Engine/Plugins/Compositing/Composite/Source/Composite/Public/Layers/CompositeLayerSceneCapture.h
@@ -25,13 +25,6 @@ public:
 	/** Destructor. */
 	UE_API virtual ~UCompositeLayerSceneCapture();
 
-	/** Set the enabled state. */
-	UE_API virtual void SetEnabled(bool bInEnabled) override;
-
-	//~ Begin UObject Interface
-	UE_API virtual void BeginDestroy() override;
-	//~ End UObject Interface
-
 #if WITH_EDITOR
 	UE_API virtual void PostEditChangeProperty(FPropertyChangedEvent& PropertyChangedEvent) override;
 #endif
diff --git a/Engine/Plugins/Compositing/Composite/Source/Composite/Public/Layers/CompositeLayerShadowReflection.h b/Engine/Plugins/Compositing/Composite/Source/Composite/Public/Layers/CompositeLayerShadowReflection.h
--- a/Engine/Plugins/Compositing/Composite/Source/Composite/Public/Layers/CompositeLayerShadowReflection.h
+++ b/Engine/Plugins/Compositing/Composite/Source/Composite/Public/Layers/CompositeLayerShadowReflection.h
@@ -41,9 +41,6 @@ public:
 	/** Destructor. */
 	UE_API virtual ~UCompositeLayerShadowReflection();
 
-	/** Set the enabled state. */
-	UE_API virtual void SetEnabled(bool bInEnabled) override;
-
 	//~ Begin UObject Interface
 	UE_API virtual void BeginDestroy() override;
 
diff --git a/Engine/Plugins/Compositing/Composite/Source/Composite/Public/Layers/CompositeLayerSingleLightShadow.h b/Engine/Plugins/Compositing/Composite/Source/Composite/Public/Layers/CompositeLayerSingleLightShadow.h
--- a/Engine/Plugins/Compositing/Composite/Source/Composite/Public/Layers/CompositeLayerSingleLightShadow.h
+++ b/Engine/Plugins/Compositing/Composite/Source/Composite/Public/Layers/CompositeLayerSingleLightShadow.h
@@ -37,12 +37,6 @@ public:
 	/** Destructor. */
 	UE_API virtual ~UCompositeLayerSingleLightShadow() = default;
 
-	/** Set the enabled state. */
-	UE_API virtual void SetEnabled(bool bInEnabled) override;
-
-	//~ Begin UObject Interface
-	UE_API virtual void BeginDestroy() override;
-
 #if WITH_EDITOR
 	UE_API virtual void PostEditChangeProperty(FPropertyChangedEvent& PropertyChangedEvent) override;
 #endif
diff --git a/Engine/Plugins/Compositing/CompositeCore/Source/CompositeCore/Private/CompositeCoreSceneViewExtension.cpp b/Engine/Plugins/Compositing/CompositeCore/Source/CompositeCore/Private/CompositeCoreSceneViewExtension.cpp
--- a/Engine/Plugins/Compositing/CompositeCore/Source/CompositeCore/Private/CompositeCoreSceneViewExtension.cpp
+++ b/Engine/Plugins/Compositing/CompositeCore/Source/CompositeCore/Private/CompositeCoreSceneViewExtension.cpp
@@ -9,6 +9,7 @@
 
 #include "CommonRenderResources.h"
 #include "Components/PrimitiveComponent.h"
+#include "Components/SceneCaptureComponent2D.h"
 #include "Containers/Set.h"
 #include "EngineUtils.h"
 #include "Engine/Texture.h"
@@ -30,13 +31,6 @@ static TAutoConsoleVariable<int32> CVarCompositeCoreApplyFXAA(
 	TEXT("When enabled, the custom render pass automatically applies FXAA."),
 	ECVF_RenderThreadSafe);
 
-static TAutoConsoleVariable<int32> CVarCompositeCoreRestrictedHoldoutViewMode(
-	TEXT("CompositeCore.RestrictedHoldoutViewMode"),
-	0,
-	TEXT(" 0: no restrictions (default)\n")
-	TEXT(" 1: primitive alpha holdout is restricted to the Lit view mode in deferred rendering.")
-);
-
 /** Parameters used to transition a texture to SRVGraphics. */
 BEGIN_SHADER_PARAMETER_STRUCT(FCompositeRenderTargetTransitionParameters, )
 	RDG_TEXTURE_ACCESS(TransitionTexture, ERHIAccess::SRVGraphics)
@@ -134,7 +128,9 @@ void FCompositeCoreSceneViewExtension::UnregisterPrimitives_GameThread(const TAr
 void FCompositeCoreSceneViewExtension::SetRenderWork_GameThread(UE::CompositeCore::FRenderWork&& InWork)
 {
 	bHasCustomRenderWork = !InWork.IsEmpty();
-	MainRenderMode = InWork.MainRenderMode;
+	MainRenderMode = MoveTemp(InWork.MainRenderMode);
+	AllowedViewModes = MoveTemp(InWork.AllowedViewModes);
+	SceneCapturesUpdateQueue = MoveTemp(InWork.SceneCapturesUpdateQueue);
 
 	ENQUEUE_RENDER_COMMAND(CopyCompositeCoreRenderWork)(
 		[FrameWork = MoveTemp(InWork), WeakThis = AsWeak()](FRHICommandListImmediate& RHICmdList) mutable
@@ -171,6 +167,8 @@ void FCompositeCoreSceneViewExtension::ResetRenderWork_GameThread()
 {
 	bHasCustomRenderWork = false;
 	MainRenderMode.Reset();
+	AllowedViewModes.Reset();
+	SceneCapturesUpdateQueue.Reset();
 
 	ENQUEUE_RENDER_COMMAND(CopyCompositeCoreRenderWork)(
 		[WeakThis = AsWeak()](FRHICommandListImmediate& RHICmdList)
@@ -240,33 +238,46 @@ void FCompositeCoreSceneViewExtension::SetupViewFamily(FSceneViewFamily& InViewF
 			Iter.RemoveCurrent();
 		}
 	}
-
-	if (MainRenderMode.IsSet())
-	{
-		InViewFamily.SceneCaptureSource = MainRenderMode.GetValue();
-	}
 }
 
 void FCompositeCoreSceneViewExtension::SetupView(FSceneViewFamily& InViewFamily, FSceneView& InView)
 {
+	const bool bIsFirstView = (&InView == InViewFamily.Views[0]);
+	if (bIsFirstView)
+	{
+		// The family view mode is only accessible after SetupViewFamily is called currently.
+		if (!AllowedViewModes.IsEmpty() && !AllowedViewModes.Contains(InViewFamily.ViewMode))
+		{
+			// Optionally disable primitive alpha holdout outside of allowed view modes
+			InViewFamily.EngineShowFlags.SetAllowPrimitiveAlphaHoldout(false);
+		}
+		else if(!InView.bIsSceneCapture)
+		{
+			// Manually update scene captures that have been registered but aren't already updating every frame
+			for (TWeakObjectPtr<USceneCaptureComponent2D>& SceneCapture : SceneCapturesUpdateQueue)
+			{
+				if (SceneCapture.IsValid() && !SceneCapture->bCaptureEveryFrame)
+				{
+					SceneCapture->CaptureSceneDeferred();
+				}
+			}
+
+			// Empty update queue
+			SceneCapturesUpdateQueue.Reset();
+		}
+	}
 }
 
 void FCompositeCoreSceneViewExtension::BeginRenderViewFamily(FSceneViewFamily& InViewFamily)
 {
-	// Optionally disable primitive alpha holdout on non-lit view modes (except path tracing)
-	// for easier workflows in editor with multiple viewports.
-	const bool bAllowPrimitiveAlphaHoldout = InViewFamily.EngineShowFlags.AllowPrimitiveAlphaHoldout;
-	const int32 bLitOnlyPrimitiveAlphaHoldout = CVarCompositeCoreRestrictedHoldoutViewMode.GetValueOnGameThread() > 0;
-
-	if (bAllowPrimitiveAlphaHoldout && bLitOnlyPrimitiveAlphaHoldout && InViewFamily.ViewMode != VMI_PathTracing)
+	if (!IsActiveForViewFamily(InViewFamily))
 	{
-		InViewFamily.EngineShowFlags.SetAllowPrimitiveAlphaHoldout(InViewFamily.ViewMode == VMI_Lit);
+		return;
 	}
 
-	// Disable custom render pass for view families without primitive alpha holdout support.
-	if (InViewFamily.EngineShowFlags.AllowPrimitiveAlphaHoldout == 0)
+	if (MainRenderMode.IsSet())
 	{
-		return;
+		InViewFamily.SceneCaptureSource = MainRenderMode.GetValue();
 	}
 
 	const TWeakObjectPtr<UWorld> WorldPtr = GetWorld();
@@ -411,10 +422,18 @@ void FCompositeCoreSceneViewExtension::SubscribeToPostProcessingPass(EPostProces
 	}
 }
 
-bool FCompositeCoreSceneViewExtension::IsActiveForView(const FSceneView& InView) const
+bool FCompositeCoreSceneViewExtension::IsActiveForViewFamily(const FSceneViewFamily& InViewFamily) const
 {
 	bool bIsActive = true;
-	bIsActive &= static_cast<bool>(InView.Family->EngineShowFlags.AllowPrimitiveAlphaHoldout);
+	bIsActive &= static_cast<bool>(InViewFamily.EngineShowFlags.AllowPrimitiveAlphaHoldout);
+	bIsActive &= AllowedViewModes.IsEmpty() ? true : AllowedViewModes.Contains(InViewFamily.ViewMode);
+
+	return bIsActive;
+}
+
+bool FCompositeCoreSceneViewExtension::IsActiveForView(const FSceneView& InView) const
+{
+	bool bIsActive = IsActiveForViewFamily(*InView.Family);
 	bIsActive &= CustomRenderTargetPerView_RenderThread.Contains(InView.GetViewKey()) || !ExternalInputs_RenderThread.IsEmpty();
 
 	return bIsActive;
diff --git a/Engine/Plugins/Compositing/CompositeCore/Source/CompositeCore/Private/CompositeCoreSceneViewExtension.h b/Engine/Plugins/Compositing/CompositeCore/Source/CompositeCore/Private/CompositeCoreSceneViewExtension.h
--- a/Engine/Plugins/Compositing/CompositeCore/Source/CompositeCore/Private/CompositeCoreSceneViewExtension.h
+++ b/Engine/Plugins/Compositing/CompositeCore/Source/CompositeCore/Private/CompositeCoreSceneViewExtension.h
@@ -84,6 +84,9 @@ private:
 	/** Getter for the frame render work, either from the actor or default. */
 	const UE::CompositeCore::FRenderWork& GetRenderWork() const;
 
+	/** Active view family check to avoid work in post-processing. */
+	bool IsActiveForViewFamily(const FSceneViewFamily& InViewFamily) const;
+
 	/** Active check to avoid work in post-processing. */
 	bool IsActiveForView(const FSceneView& InView) const;
 
@@ -132,6 +135,12 @@ private:
 	/** Optional main render mode override */
 	TOptional<ESceneCaptureSource> MainRenderMode;
 
+	/** View modes for which the compositing is allowed. */
+	TSet<EViewModeIndex> AllowedViewModes;
+
+	/** List of scene captures to render-update this frame. */
+	TArray<TWeakObjectPtr<USceneCaptureComponent2D>, SceneRenderingAllocator> SceneCapturesUpdateQueue;
+
 	/**
 	* Map of cached proxy output textures per view (post-processing only currently).
 	* 
diff --git a/Engine/Plugins/Compositing/CompositeCore/Source/CompositeCore/Private/Passes/CompositeCorePassProxy.cpp b/Engine/Plugins/Compositing/CompositeCore/Source/CompositeCore/Private/Passes/CompositeCorePassProxy.cpp
--- a/Engine/Plugins/Compositing/CompositeCore/Source/CompositeCore/Private/Passes/CompositeCorePassProxy.cpp
+++ b/Engine/Plugins/Compositing/CompositeCore/Source/CompositeCore/Private/Passes/CompositeCorePassProxy.cpp
@@ -79,6 +79,7 @@ namespace UE
 			}
 
 			bIsEmpty &= !MainRenderMode.IsSet();
+			bIsEmpty &= SceneCapturesUpdateQueue.IsEmpty();
 
 			return bIsEmpty;
 		}
diff --git a/Engine/Plugins/Compositing/CompositeCore/Source/CompositeCore/Public/Passes/CompositeCorePassProxy.h b/Engine/Plugins/Compositing/CompositeCore/Source/CompositeCore/Public/Passes/CompositeCorePassProxy.h
--- a/Engine/Plugins/Compositing/CompositeCore/Source/CompositeCore/Public/Passes/CompositeCorePassProxy.h
+++ b/Engine/Plugins/Compositing/CompositeCore/Source/CompositeCore/Public/Passes/CompositeCorePassProxy.h
@@ -23,6 +23,7 @@ struct FPostProcessMaterialInputs;
 class FCompositeCorePassProxy;
 struct IPooledRenderTarget;
 class UTexture;
+class USceneCaptureComponent2D;
 
 enum class EPostProcessMaterialInput : uint32;
 
@@ -216,7 +217,7 @@ namespace UE
 			bool IsEmpty() const;
 
 			/** Array of user-defined external input texture overrides, where indices map to ResourceId. */
-			TArray<FExternalTexture> ExternalInputs;
+			TArray<FExternalTexture, SceneRenderingAllocator> ExternalInputs;
 
 			/** Array of pre-processing passes applied on specific texture resources at the start of rendering. */
 			TSortedMap<ResourceId, TArray<const FCompositeCorePassProxy*>> PreprocessingPasses;
@@ -229,6 +230,12 @@ namespace UE
 
 			/** Optional main render mode override */
 			TOptional<ESceneCaptureSource> MainRenderMode;
+
+			/** View modes for which the compositing is allowed. */
+			TSet<EViewModeIndex> AllowedViewModes;
+
+			/** List of scene captures to render-update this frame. */
+			TArray<TWeakObjectPtr<USceneCaptureComponent2D>, SceneRenderingAllocator> SceneCapturesUpdateQueue;
 		};
 	}
 }
