CompositeCore: Fixed concurrency issue introduced in 46310920.

Allowed view modes were access from both game and render threads. As a fix, we branch copy access per thread.

#rnx

diff --git a/Engine/Plugins/Compositing/CompositeCore/Source/CompositeCore/Private/CompositeCoreSceneViewExtension.cpp b/Engine/Plugins/Compositing/CompositeCore/Source/CompositeCore/Private/CompositeCoreSceneViewExtension.cpp
--- a/Engine/Plugins/Compositing/CompositeCore/Source/CompositeCore/Private/CompositeCoreSceneViewExtension.cpp
+++ b/Engine/Plugins/Compositing/CompositeCore/Source/CompositeCore/Private/CompositeCoreSceneViewExtension.cpp
@@ -129,7 +129,8 @@ void FCompositeCoreSceneViewExtension::SetRenderWork_GameThread(UE::CompositeCor
 {
 	bHasCustomRenderWork = !InWork.IsEmpty();
 	MainRenderMode = MoveTemp(InWork.MainRenderMode);
-	AllowedViewModes = MoveTemp(InWork.AllowedViewModes);
+	// Copy since view modes are used on both game and render threads (see IsActiveForViewFamily)
+	AllowedViewModes = InWork.AllowedViewModes;
 	SceneCapturesUpdateQueue = MoveTemp(InWork.SceneCapturesUpdateQueue);
 
 	ENQUEUE_RENDER_COMMAND(CopyCompositeCoreRenderWork)(
@@ -424,9 +425,26 @@ void FCompositeCoreSceneViewExtension::SubscribeToPostProcessingPass(EPostProces
 
 bool FCompositeCoreSceneViewExtension::IsActiveForViewFamily(const FSceneViewFamily& InViewFamily) const
 {
+	auto IsViewModeAllowedFn = [&InViewFamily](const TSet<EViewModeIndex>& InAllowedViewModes) -> bool
+		{
+			// If the allowed view modes set is empty, we allow all types.
+			return InAllowedViewModes.IsEmpty() ? true : InAllowedViewModes.Contains(InViewFamily.ViewMode);
+		};
+
 	bool bIsActive = true;
 	bIsActive &= static_cast<bool>(InViewFamily.EngineShowFlags.AllowPrimitiveAlphaHoldout);
-	bIsActive &= AllowedViewModes.IsEmpty() ? true : AllowedViewModes.Contains(InViewFamily.ViewMode);
+
+	if (IsInRenderingThread())
+	{
+		if (RenderWork_RenderThread.IsSet())
+		{
+			bIsActive &= IsViewModeAllowedFn(RenderWork_RenderThread.GetValue().AllowedViewModes);
+		}
+	}
+	else if(IsInGameThread())
+	{
+		bIsActive &= IsViewModeAllowedFn(AllowedViewModes);
+	}
 
 	return bIsActive;
 }
