Composite: Fix holdout state change no registering through MU.

diff --git a/Engine/Plugins/Compositing/CompositeCore/Source/CompositeCore/Private/CompositeCoreSceneViewExtension.cpp b/Engine/Plugins/Compositing/CompositeCore/Source/CompositeCore/Private/CompositeCoreSceneViewExtension.cpp
--- a/Engine/Plugins/Compositing/CompositeCore/Source/CompositeCore/Private/CompositeCoreSceneViewExtension.cpp
+++ b/Engine/Plugins/Compositing/CompositeCore/Source/CompositeCore/Private/CompositeCoreSceneViewExtension.cpp
@@ -99,8 +99,11 @@ void FCompositeCoreSceneViewExtension::RegisterPrimitives_GameThread(const TArra
 			CompositePrimitives.Add(InPrimitiveComponent);
 		}
 
-		// The SetHoldout() function makes changes only if the HoldoutState value differs from the current value.
-		InPrimitiveComponent->SetHoldout(true);
+		if (!InPrimitiveComponent->bHoldout)
+		{
+			InPrimitiveComponent->Modify();
+			InPrimitiveComponent->SetHoldout(true);
+		}
 	}
 }
 
@@ -120,8 +123,11 @@ void FCompositeCoreSceneViewExtension::UnregisterPrimitives_GameThread(const TAr
 			CompositePrimitives.Remove(InPrimitiveComponent);
 		}
 
-		// The SetHoldout() function makes changes only if the HoldoutState value differs from the current value.
-		InPrimitiveComponent->SetHoldout(false);
+		if (InPrimitiveComponent->bHoldout)
+		{
+			InPrimitiveComponent->Modify();
+			InPrimitiveComponent->SetHoldout(false);
+		}
 	}
 }
 
