Composite: Fixed RHI validation errors and preventing failing ensures.

* Fixed missing transitions (into UAV for compute) on clean plate keyer texture.
* Unified SVE external textures usage into a scoped struct for internal access.
* Fixed unwanted ensure trigger when no pass is active.
* Logged a warning on missing external texture ensure (punting to 5.8 for a larger scope fix).

#rnx

diff --git a/Engine/Plugins/Compositing/Composite/Source/Composite/Private/CompositeActor.cpp b/Engine/Plugins/Compositing/Composite/Source/Composite/Private/CompositeActor.cpp
--- a/Engine/Plugins/Compositing/Composite/Source/Composite/Private/CompositeActor.cpp
+++ b/Engine/Plugins/Compositing/Composite/Source/Composite/Private/CompositeActor.cpp
@@ -165,10 +165,17 @@ namespace UE
 				Context.bIsSolo = (Layers.Num() == 1);
 
 				const int32 BranchLayerIndex = RegisterLayers(Context, Layers, *RenderWork.FrameAllocator, LayerProxies);
-				const int32 BranchSourceIndex = BranchLayerIndex - 1;
+				
+				if (LayerProxies.IsEmpty())
+				{
+					// Early return when there is no actively rendering layer
+					return {};
+				}
 
 				if (InCompositeActor.bEnableScreenSpaceReflections)
 				{
+					const int32 BranchSourceIndex = BranchLayerIndex - 1;
+
 					TArray<FCompositeCorePassProxy*> BranchedLayerProxies;
 					Context.bIsFirstPass = true;
 					RegisterLayers(Context, Layers, *RenderWork.FrameAllocator, BranchedLayerProxies);
diff --git a/Engine/Plugins/Compositing/Composite/Source/Composite/Private/Passes/CompositePassColorKeyer.cpp b/Engine/Plugins/Compositing/Composite/Source/Composite/Private/Passes/CompositePassColorKeyer.cpp
--- a/Engine/Plugins/Compositing/Composite/Source/Composite/Private/Passes/CompositePassColorKeyer.cpp
+++ b/Engine/Plugins/Compositing/Composite/Source/Composite/Private/Passes/CompositePassColorKeyer.cpp
@@ -5,6 +5,7 @@
 #include "Engine/Texture.h"
 #include "Passes/CompositeCorePassProxy.h"
 #include "PixelShaderUtils.h"
+#include "SystemTextures.h"
 #include "TextureResource.h"
 
 #if WITH_EDITOR
@@ -46,7 +47,7 @@ class FCompositePassColorKeyerShader : public FGlobalShader
 	BEGIN_SHADER_PARAMETER_STRUCT(FParameters, )
 		SHADER_PARAMETER_RDG_TEXTURE(Texture2D, InputTexture)
 		SHADER_PARAMETER_SAMPLER(SamplerState, InputSampler)
-		SHADER_PARAMETER_TEXTURE(Texture2D, CleanPlateTexture)
+		SHADER_PARAMETER_RDG_TEXTURE(Texture2D, CleanPlateTexture)
 		SHADER_PARAMETER_SAMPLER(SamplerState, CleanPlateSampler)
 		SHADER_PARAMETER_RDG_TEXTURE_UAV(RWTexture2D<float4>, RWOutputTexture)
 		SHADER_PARAMETER_STRUCT(FScreenPassTextureViewportParameters, Input)
@@ -155,21 +156,29 @@ namespace UE
 				PassParameters->bPreserveVignetteAfterKey = static_cast<uint32>(bPreserveVignetteAfterKey);
 				PassParameters->bInvertAlpha = static_cast<uint32>(bInvertAlpha);
 	
-				bool bUseCleanPlate = false;
-				PassParameters->CleanPlateTexture = GBlackTexture->TextureRHI;
+				FRDGTextureRef CleanPlateResource = nullptr;
+				PassParameters->CleanPlateTexture = GSystemTextures.GetBlackDummy(GraphBuilder);
 				PassParameters->CleanPlateSampler = TStaticSamplerState<SF_Bilinear>::GetRHI();
 
-				if (CleanPlateWeakPtr.IsValid(false, true))
+				TStrongObjectPtr<UTexture> CleanPlateTexturePtr = CleanPlateWeakPtr.Pin();
+				if (CleanPlateTexturePtr.IsValid())
 				{
-					if (const FTextureResource* TextureResource = CleanPlateWeakPtr->GetResource())
+					const FTextureResource* TextureResource = CleanPlateTexturePtr->GetResource();
+					if (TextureResource != nullptr)
 					{
-						PassParameters->CleanPlateTexture = TextureResource->TextureRHI;
-						bUseCleanPlate = true;
+						FRHITexture* TextureRHI = TextureResource->GetTexture2DRHI();
+						if (TextureRHI != nullptr)
+						{
+							CleanPlateResource = GraphBuilder.RegisterExternalTexture(CreateRenderTarget(TextureRHI, TEXT("CompositeCleanPlateTexture")));
+							GraphBuilder.UseInternalAccessMode(CleanPlateResource);
+
+							PassParameters->CleanPlateTexture = CleanPlateResource;
+						}
 					}
 				}
 
 				FCompositePassColorKeyerShader::FPermutationDomain PermutationVector;
-				PermutationVector.Set<FCompositePassColorKeyerShader::FUseCleanPlate>(bUseCleanPlate);
+				PermutationVector.Set<FCompositePassColorKeyerShader::FUseCleanPlate>(CleanPlateResource != nullptr);
 
 				const FIntPoint OutputSize = Output.ViewRect.Size();
 				TShaderMapRef<FCompositePassColorKeyerShader> ComputeShader(GlobalShaderMap, PermutationVector);
@@ -182,6 +191,11 @@ namespace UE
 					FComputeShaderUtils::GetGroupCount(OutputSize, FCompositePassColorKeyerShader::ThreadGroupSize)
 				);
 
+				if (CleanPlateResource != nullptr)
+				{
+					GraphBuilder.UseExternalAccessMode(CleanPlateResource, ERHIAccess::SRVMask);
+				}
+
 				return UE::CompositeCore::FPassTexture{ MoveTemp(Output), Metadata };
 			}
 
diff --git a/Engine/Plugins/Compositing/CompositeCore/Source/CompositeCore/Private/CompositeCoreSceneViewExtension.cpp b/Engine/Plugins/Compositing/CompositeCore/Source/CompositeCore/Private/CompositeCoreSceneViewExtension.cpp
--- a/Engine/Plugins/Compositing/CompositeCore/Source/CompositeCore/Private/CompositeCoreSceneViewExtension.cpp
+++ b/Engine/Plugins/Compositing/CompositeCore/Source/CompositeCore/Private/CompositeCoreSceneViewExtension.cpp
@@ -475,15 +475,15 @@ void FCompositeCoreSceneViewExtension::ApplyPreprocessing(FRDGBuilder& GraphBuil
 	}
 
 	FPassContext PassContext = {};
-	const FExternalTextureMap ExternalTextures = CreateExternalTextureMap(GraphBuilder, InView);
+	const FScopedExternalTextureMap ExternalTextures{ *this, GraphBuilder, InView };
 
 	for (const TPair<ResourceId, TArray<const FCompositeCorePassProxy*>>& Pair : PreprocessingPasses)
 	{
 		const ResourceId ExternalId = Pair.Key;
 		const TArray<const FCompositeCorePassProxy*>& Passes = Pair.Value;
-		const FPassInput* ExternalTexture = ExternalTextures.Find(ExternalId);
+		const FPassInput* ExternalTexture = ExternalTextures.Get().Find(ExternalId);
 
-		if (!ensureMsgf(ExternalTexture != nullptr, TEXT("Unexpected missing external texture override as output.")))
+		if (!ensureMsgf(ExternalTexture != nullptr, TEXT("Unexpected missing external texture.")))
 		{
 			continue;
 		}
@@ -496,25 +496,25 @@ void FCompositeCoreSceneViewExtension::ApplyPreprocessing(FRDGBuilder& GraphBuil
 		for (int32 PassIndex =0; PassIndex <Passes.Num(); ++PassIndex)
 		{
 			const FCompositeCorePassProxy* Pass = Passes[PassIndex];
-
-			FRDGTextureRef ExternalRDG = nullptr;
 			const TOptional<UE::CompositeCore::ResourceId>& OverrrideOutput = Pass->GetDeclaredOutputOverride();
-			
+			bool bPassOutputOverridden = false;
+
 			if (OverrrideOutput.IsSet())
 			{
-				const FExternalRenderTarget* ExternalRT = ExternalInputs_RenderThread.Find(OverrrideOutput.GetValue());
-				if (ensureMsgf(ExternalRT != nullptr, TEXT("Unexpected missing external render target override as output.")))
+				const FPassInput* ExternalOutputOverride = ExternalTextures.Get().Find(OverrrideOutput.GetValue());
+				if (ensureMsgf(ExternalOutputOverride != nullptr, TEXT("Unexpected missing external render target override as output.")))
 				{
-					ExternalRDG = GraphBuilder.RegisterExternalTexture(ExternalRT->RenderTarget);
-					GraphBuilder.UseInternalAccessMode(ExternalRDG);
-
-					PassInputs.OverrideOutput = FScreenPassRenderTarget(ExternalRDG, ERenderTargetLoadAction::ENoAction);
-
-					// Update view rect to match output
-					PassContext.OutputViewRect = FIntRect(FInt32Point::ZeroValue, ExternalRT->RenderTarget->GetDesc().Extent);
+					PassInputs.OverrideOutput = FScreenPassRenderTarget(ExternalOutputOverride->Texture, ERenderTargetLoadAction::ENoAction);
+					
+					// Update view rect to match external texture
+					PassContext.OutputViewRect = ExternalOutputOverride->Texture.ViewRect;
+					// Mark as overriden
+					bPassOutputOverridden = true;
 				}
 			}
-			else
+			
+
+			if(!bPassOutputOverridden)
 			{
 				PassInputs.OverrideOutput = FScreenPassRenderTarget{};
 
@@ -527,46 +527,57 @@ void FCompositeCoreSceneViewExtension::ApplyPreprocessing(FRDGBuilder& GraphBuil
 
 			// Note: we do not cache the output to CachedOutputsPerView_RenderThread since we only support cached outputs in post-processing for now.
 
-			if (ExternalRDG)
-			{
-				// Transition the render target back to external access as a shader resource.
-				GraphBuilder.UseExternalAccessMode(ExternalRDG, ERHIAccess::SRVMask);
-			}
-
 			// Update next pass input
 			PassInputs[0] = Output;
 		}
 	}
 }
 
-TSortedMap<int32, UE::CompositeCore::FPassInput> FCompositeCoreSceneViewExtension::CreateExternalTextureMap(FRDGBuilder& GraphBuilder, const FSceneView& InView) const
+UE::CompositeCore::FScopedExternalTextureMap::FScopedExternalTextureMap(const FCompositeCoreSceneViewExtension& SVE, FRDGBuilder& InGraphBuilder, const FSceneView& InView)
+	: GraphBuilder{InGraphBuilder}
 {
 	using namespace UE::CompositeCore;
 
-	TSortedMap<int32, FPassInput> InputOverrides;
-	InputOverrides.Reserve(2 + ExternalInputs_RenderThread.Num());
+	Textures.Reserve(2 + SVE.ExternalInputs_RenderThread.Num());
 
 	{
-		FPassInput& Resource = InputOverrides.Add(BUILT_IN_EMPTY_ID);
+		FPassInput& Resource = Textures.Add(BUILT_IN_EMPTY_ID);
 		Resource.Texture = FScreenPassTexture{ GSystemTextures.GetBlackDummy(GraphBuilder) };
 	}
 
-	const TRefCountPtr<IPooledRenderTarget>* CompositeRenderPassPtr = CustomRenderTargetPerView_RenderThread.Find(InView.GetViewKey());
+	const TRefCountPtr<IPooledRenderTarget>* CompositeRenderPassPtr = SVE.CustomRenderTargetPerView_RenderThread.Find(InView.GetViewKey());
 	if (CompositeRenderPassPtr != nullptr)
 	{
-		FPassInput& Resource = InputOverrides.Add(BUILT_IN_CRP_ID);
+		FPassInput& Resource = Textures.Add(BUILT_IN_CRP_ID);
 		Resource.Texture = FScreenPassTexture{ GraphBuilder.RegisterExternalTexture(*CompositeRenderPassPtr) };
 		Resource.Metadata.bInvertedAlpha = true;
 	}
 
-	for (const auto& Pair : ExternalInputs_RenderThread)
+	RestoreToExternal.Reserve(SVE.ExternalInputs_RenderThread.Num());
+
+	for (const auto& Pair : SVE.ExternalInputs_RenderThread)
 	{
-		FPassInput& Resource = InputOverrides.Add(Pair.Key);
-		Resource.Texture = FScreenPassTexture{ GraphBuilder.RegisterExternalTexture(Pair.Value.RenderTarget) };
+		FPassInput& Resource = Textures.Add(Pair.Key);
+		FRDGTextureRef TextureResource = GraphBuilder.RegisterExternalTexture(Pair.Value.RenderTarget);
+		Resource.Texture = FScreenPassTexture{ TextureResource };
 		Resource.Metadata = Pair.Value.Metadata;
+
+		// Mark for internal access and store to reset in destructor
+		if (TextureResource->IsExternal())
+		{
+			GraphBuilder.UseInternalAccessMode(TextureResource);
+
+			RestoreToExternal.Add(TextureResource);
+		}
 	}
+}
 
-	return InputOverrides;
+UE::CompositeCore::FScopedExternalTextureMap::~FScopedExternalTextureMap()
+{
+	for (FRDGViewableResource* TextureResource : RestoreToExternal)
+	{
+		GraphBuilder.UseExternalAccessMode(TextureResource, ERHIAccess::SRVMask);
+	}
 }
 
 bool FCompositeCoreSceneViewExtension::ApplyPasses_Recursive(
@@ -576,7 +587,7 @@ bool FCompositeCoreSceneViewExtension::ApplyPasses_Recursive(
 	const UE::CompositeCore::FPassInputArray& OriginalInputs,
 	UE::CompositeCore::FPassContext& PassContext,
 	const TArray<const FCompositeCorePassProxy*> InPasses,
-	FExternalTextureMap& ExternalTextures,
+	const UE::CompositeCore::FScopedExternalTextureMap& InExternalTextures,
 	const int32 LastMergePassIndex,
 	const int32 RecursionLevel,
 	UE::CompositeCore::FPassTexture& Output)
@@ -627,7 +638,7 @@ bool FCompositeCoreSceneViewExtension::ApplyPasses_Recursive(
 				{
 					PassInputs[InputIndex] = *PreProducedOutput;
 				}
-				else if(ApplyPasses_Recursive(GraphBuilder, InView, PassInputs, OriginalInputs, PassContext, { ChildPassProxy }, ExternalTextures, INDEX_NONE, RecursionLevel + 1, ChildPassOutput))
+				else if(ApplyPasses_Recursive(GraphBuilder, InView, PassInputs, OriginalInputs, PassContext, { ChildPassProxy }, InExternalTextures, INDEX_NONE, RecursionLevel + 1, ChildPassOutput))
 				{
 					PassInputs[InputIndex] = ChildPassOutput;
 				}
@@ -638,11 +649,18 @@ bool FCompositeCoreSceneViewExtension::ApplyPasses_Recursive(
 				const FPassExternalResourceDesc& Desc = DeclaredInput.Get<FPassExternalResourceDesc>();
 				const ResourceId DeclaredExternalId = Desc.Id;
 
-				const FPassInput* SliceInput = ExternalTextures.Find(DeclaredExternalId);
+				const FPassInput* SliceInput = InExternalTextures.Get().Find(DeclaredExternalId);
 				if (ensureMsgf(SliceInput, TEXT("Invalid external input: %d"), DeclaredExternalId))
 				{
 					PassInputs[InputIndex] = *SliceInput;
 				}
+				else
+				{
+					UE_CALL_ONCE([]()
+						{
+							UE_LOG(LogCompositeCore, Warning, TEXT("Compositing may be incorrectly running inside a scene capture render, and trying to use render targets which have not yet been updated. Consider disabling AllowPrimitiveAlphaHoldout on your scene capture(s)."))
+						});
+				}
 			}
 			// If an internal texture resource is expected, connect the current input index with either the original or current bass pass input.
 			else if (DeclaredInput.IsType<FPassInternalResourceDesc>())
@@ -699,7 +717,7 @@ FScreenPassTexture FCompositeCoreSceneViewExtension::PostProcessWork_RenderThrea
 {
 	using namespace UE::CompositeCore;
 
-	FExternalTextureMap ExternalTextures = CreateExternalTextureMap(GraphBuilder, InView);
+	const FScopedExternalTextureMap ExternalTextures{ *this, GraphBuilder, InView };
 	const TArray<const FCompositeCorePassProxy*>& Passes = GetRenderWork().FramePasses[InLocation];
 
 	FPassContext PassContext;
@@ -736,4 +754,3 @@ void FCompositeCoreSceneViewExtension::PostRenderView_RenderThread(FRDGBuilder&
 {
 	CustomRenderTargetPerView_RenderThread.Remove(InView.GetViewKey());
 }
-
diff --git a/Engine/Plugins/Compositing/CompositeCore/Source/CompositeCore/Private/CompositeCoreSceneViewExtension.h b/Engine/Plugins/Compositing/CompositeCore/Source/CompositeCore/Private/CompositeCoreSceneViewExtension.h
--- a/Engine/Plugins/Compositing/CompositeCore/Source/CompositeCore/Private/CompositeCoreSceneViewExtension.h
+++ b/Engine/Plugins/Compositing/CompositeCore/Source/CompositeCore/Private/CompositeCoreSceneViewExtension.h
@@ -12,6 +12,7 @@
 #include "Passes/CompositeCorePassProxy.h"
 
 class FViewInfo;
+class FCompositeCoreSceneViewExtension;
 struct FScreenPassRenderTarget;
 struct FScreenPassTexture;
 struct FScreenPassTextureSlice;
@@ -30,6 +31,32 @@ namespace UE
 			TRefCountPtr<IPooledRenderTarget> RenderTarget = {};
 			FResourceMetadata Metadata = {};
 		};
+
+		/** Collection of external textures scoped to internal RDG access. */
+		struct FScopedExternalTextureMap
+		{
+			UE_NONCOPYABLE(FScopedExternalTextureMap);
+
+			/** Constructor */
+			FScopedExternalTextureMap(const FCompositeCoreSceneViewExtension& SVE, FRDGBuilder& InGraphBuilder, const FSceneView& InView);
+			
+			/** Destructor */
+			~FScopedExternalTextureMap();
+
+			/** Get the map of key indices to pass input textures. */
+			const TSortedMap<int32, UE::CompositeCore::FPassInput>& Get() const
+			{
+				return Textures;
+			}
+
+		private:
+			/** Host graph builder */
+			FRDGBuilder& GraphBuilder;
+			/** Map of pass input external textures. */
+			TSortedMap<int32, UE::CompositeCore::FPassInput> Textures;
+			/** List of textures marked for internal access for external restoration upon destruction. */
+			TArray<FRDGViewableResource*, SceneRenderingAllocator> RestoreToExternal;
+		};
 	}
 }
 
@@ -79,8 +106,6 @@ protected:
 	//~ End ISceneViewExtension Interface
 
 private:
-	using FExternalTextureMap = TSortedMap<int32, UE::CompositeCore::FPassInput>;
-
 	/** Getter for the frame render work, either from the actor or default. */
 	const UE::CompositeCore::FRenderWork& GetRenderWork() const;
 
@@ -93,9 +118,6 @@ private:
 	/** Once-per frame external resource processing passes. */
 	void ApplyPreprocessing(FRDGBuilder& GraphBuilder, FSceneView& InView) const;
 
-	/** Convenience function to create an override map of external textures. */
-	FExternalTextureMap CreateExternalTextureMap(FRDGBuilder& GraphBuilder, const FSceneView& InView) const;
-
 	/** Recursive pass application. */
 	bool ApplyPasses_Recursive(
 		FRDGBuilder& GraphBuilder,
@@ -104,7 +126,7 @@ private:
 		const UE::CompositeCore::FPassInputArray& OriginalInputs,
 		UE::CompositeCore::FPassContext& PassContext,
 		const TArray<const FCompositeCorePassProxy*> InPasses,
-		FExternalTextureMap& InOutExternalTextures,
+		const UE::CompositeCore::FScopedExternalTextureMap& InExternalTextures,
 		const int32 LastMergePassIndex,
 		const int32 RecursionLevel,
 		UE::CompositeCore::FPassTexture& Output
@@ -150,5 +172,7 @@ private:
 	* when it is later used as an input of a subsequent pass.
 	*/
 	mutable TSortedMap<uint32, TSortedMap<const FCompositeCorePassProxy*, const UE::CompositeCore::FPassTexture>> CachedOutputsPerView_RenderThread;
+
+	friend struct UE::CompositeCore::FScopedExternalTextureMap;
 };
 
