Composite: Fix unwanted transition to external on active scene capture render target(s).

Previous CL 46841795 expanded the use of internal/extermal access management to fix RHI transition validation errors. However, we should not alter the active view family render target, otherwise "Cannot set final access on a resource in external access mode" is triggered.

#rnx

diff --git a/Engine/Plugins/Compositing/CompositeCore/Source/CompositeCore/Private/CompositeCoreSceneViewExtension.cpp b/Engine/Plugins/Compositing/CompositeCore/Source/CompositeCore/Private/CompositeCoreSceneViewExtension.cpp
--- a/Engine/Plugins/Compositing/CompositeCore/Source/CompositeCore/Private/CompositeCoreSceneViewExtension.cpp
+++ b/Engine/Plugins/Compositing/CompositeCore/Source/CompositeCore/Private/CompositeCoreSceneViewExtension.cpp
@@ -562,8 +562,17 @@ UE::CompositeCore::FScopedExternalTextureMap::FScopedExternalTextureMap(const FC
 		Resource.Texture = FScreenPassTexture{ TextureResource };
 		Resource.Metadata = Pair.Value.Metadata;
 
-		// Mark for internal access and store to reset in destructor
-		if (TextureResource->IsExternal())
+		/**
+		 * Mark for internal access and store to reset in destructor.
+		 * 
+		 * Note: we exclude the active view family render target since it is already (and should remain) internal.
+		 * In practice, this can occur when our textures are accessed by pre-processing passes during a scene capture render.
+		 * 
+		 * Unfortunately, full AccessModeState information is not accessible outside of RDG or friend classes. See for example
+		 * GraphBuilder.UseExternalAccessMode(Buffer, AccessModeState.Access, AccessModeState.Pipelines); use in FRDGResourceDumpContext.
+		 */
+		FRHITexture* FamilyRenderTargetRHI = InView.Family->RenderTarget ? InView.Family->RenderTarget->GetRenderTargetTexture().GetReference() : nullptr;
+		if (TextureResource && (TextureResource->GetRHI() != FamilyRenderTargetRHI))
 		{
 			GraphBuilder.UseInternalAccessMode(TextureResource);
 
