CompositeCore: Fixed processing layer not running without additional compositing work.

* Remove residual activation conditions from "holdout composite" older implementation. We only keep the "allowed" conditions.

diff --git a/Engine/Plugins/Compositing/CompositeCore/Source/CompositeCore/Private/CompositeCoreSceneViewExtension.cpp b/Engine/Plugins/Compositing/CompositeCore/Source/CompositeCore/Private/CompositeCoreSceneViewExtension.cpp
--- a/Engine/Plugins/Compositing/CompositeCore/Source/CompositeCore/Private/CompositeCoreSceneViewExtension.cpp
+++ b/Engine/Plugins/Compositing/CompositeCore/Source/CompositeCore/Private/CompositeCoreSceneViewExtension.cpp
@@ -277,7 +277,7 @@ void FCompositeCoreSceneViewExtension::SetupView(FSceneViewFamily& InViewFamily,
 
 void FCompositeCoreSceneViewExtension::BeginRenderViewFamily(FSceneViewFamily& InViewFamily)
 {
-	if (!IsActiveForViewFamily(InViewFamily))
+	if (!IsAllowedForViewFamily(InViewFamily))
 	{
 		return;
 	}
@@ -418,7 +418,7 @@ void FCompositeCoreSceneViewExtension::PrePostProcessPass_RenderThread(FRDGBuild
 
 void FCompositeCoreSceneViewExtension::SubscribeToPostProcessingPass(EPostProcessingPass PassId, const FSceneView& InView, FAfterPassCallbackDelegateArray& InOutPassCallbacks, bool bIsPassEnabled)
 {
-	if (!IsActiveForView(InView))
+	if (!IsAllowedForViewFamily(*InView.Family))
 	{
 		return;
 	}
@@ -429,7 +429,7 @@ void FCompositeCoreSceneViewExtension::SubscribeToPostProcessingPass(EPostProces
 	}
 }
 
-bool FCompositeCoreSceneViewExtension::IsActiveForViewFamily(const FSceneViewFamily& InViewFamily) const
+bool FCompositeCoreSceneViewExtension::IsAllowedForViewFamily(const FSceneViewFamily& InViewFamily) const
 {
 	auto IsViewModeAllowedFn = [&InViewFamily](const TSet<EViewModeIndex>& InAllowedViewModes) -> bool
 		{
@@ -437,30 +437,19 @@ bool FCompositeCoreSceneViewExtension::IsActiveForViewFamily(const FSceneViewFam
 			return InAllowedViewModes.IsEmpty() ? true : InAllowedViewModes.Contains(InViewFamily.ViewMode);
 		};
 
-	bool bIsActive = true;
-	bIsActive &= static_cast<bool>(InViewFamily.EngineShowFlags.AllowPrimitiveAlphaHoldout);
+	bool bIsAllowed = true;
+	bIsAllowed &= static_cast<bool>(InViewFamily.EngineShowFlags.AllowPrimitiveAlphaHoldout);
 
 	if (IsInRenderingThread())
 	{
-		if (RenderWork_RenderThread.IsSet())
-		{
-			bIsActive &= IsViewModeAllowedFn(RenderWork_RenderThread.GetValue().AllowedViewModes);
-		}
+		bIsAllowed &= IsViewModeAllowedFn(GetRenderWork().AllowedViewModes);
 	}
 	else if(IsInGameThread())
 	{
-		bIsActive &= IsViewModeAllowedFn(AllowedViewModes);
+		bIsAllowed &= IsViewModeAllowedFn(AllowedViewModes);
 	}
 
-	return bIsActive;
-}
-
-bool FCompositeCoreSceneViewExtension::IsActiveForView(const FSceneView& InView) const
-{
-	bool bIsActive = IsActiveForViewFamily(*InView.Family);
-	bIsActive &= CustomRenderTargetPerView_RenderThread.Contains(InView.GetViewKey()) || !ExternalInputs_RenderThread.IsEmpty();
-
-	return bIsActive;
+	return bIsAllowed;
 }
 
 void FCompositeCoreSceneViewExtension::ApplyPreprocessing(FRDGBuilder& GraphBuilder, FSceneView& InView) const
@@ -545,12 +534,20 @@ UE::CompositeCore::FScopedExternalTextureMap::FScopedExternalTextureMap(const FC
 		Resource.Texture = FScreenPassTexture{ GSystemTextures.GetBlackDummy(GraphBuilder) };
 	}
 
-	const TRefCountPtr<IPooledRenderTarget>* CompositeRenderPassPtr = SVE.CustomRenderTargetPerView_RenderThread.Find(InView.GetViewKey());
-	if (CompositeRenderPassPtr != nullptr)
 	{
 		FPassInput& Resource = Textures.Add(BUILT_IN_CRP_ID);
-		Resource.Texture = FScreenPassTexture{ GraphBuilder.RegisterExternalTexture(*CompositeRenderPassPtr) };
-		Resource.Metadata.bInvertedAlpha = true;
+
+		const TRefCountPtr<IPooledRenderTarget>* CompositeRenderPassPtr = SVE.CustomRenderTargetPerView_RenderThread.Find(InView.GetViewKey());
+		if (CompositeRenderPassPtr != nullptr)
+		{
+			Resource.Texture = FScreenPassTexture{ GraphBuilder.RegisterExternalTexture(*CompositeRenderPassPtr) };
+			Resource.Metadata.bInvertedAlpha = true;
+		}
+		else
+		{
+			// Still provide an empty CRP target
+			Resource.Texture = FScreenPassTexture{ GSystemTextures.GetBlackDummy(GraphBuilder) };
+		}
 	}
 
 	RestoreToExternal.Reserve(SVE.ExternalInputs_RenderThread.Num());
diff --git a/Engine/Plugins/Compositing/CompositeCore/Source/CompositeCore/Private/CompositeCoreSceneViewExtension.h b/Engine/Plugins/Compositing/CompositeCore/Source/CompositeCore/Private/CompositeCoreSceneViewExtension.h
--- a/Engine/Plugins/Compositing/CompositeCore/Source/CompositeCore/Private/CompositeCoreSceneViewExtension.h
+++ b/Engine/Plugins/Compositing/CompositeCore/Source/CompositeCore/Private/CompositeCoreSceneViewExtension.h
@@ -109,11 +109,8 @@ private:
 	/** Getter for the frame render work, either from the actor or default. */
 	const UE::CompositeCore::FRenderWork& GetRenderWork() const;
 
-	/** Active view family check to avoid work in post-processing. */
-	bool IsActiveForViewFamily(const FSceneViewFamily& InViewFamily) const;
-
-	/** Active check to avoid work in post-processing. */
-	bool IsActiveForView(const FSceneView& InView) const;
+	/** Check if render work is allowed on the view family. */
+	bool IsAllowedForViewFamily(const FSceneViewFamily& InViewFamily) const;
 
 	/** Once-per frame external resource processing passes. */
 	void ApplyPreprocessing(FRDGBuilder& GraphBuilder, FSceneView& InView) const;
