CompositeCore: Fixed dx11 validation error, and improve debugging with resource names.

#Synced-CL 48399844

diff --git a/Engine/Plugins/Compositing/Composite/Source/Composite/Private/Layers/CompositeLayerPlate.cpp b/Engine/Plugins/Compositing/Composite/Source/Composite/Private/Layers/CompositeLayerPlate.cpp
--- a/Engine/Plugins/Compositing/Composite/Source/Composite/Private/Layers/CompositeLayerPlate.cpp
+++ b/Engine/Plugins/Compositing/Composite/Source/Composite/Private/Layers/CompositeLayerPlate.cpp
@@ -456,6 +456,8 @@ UE::CompositeCore::ResourceId UCompositeLayerPlate::AddPreprocessingPasses(
 			{
 				FResourceMetadata Metadata = {};
 				Metadata.bDistorted = true;
+				Metadata.DebugName = TEXT("CompositePlateTex");
+
 				OutputTextureId = InContext.FindOrCreateExternalTexture(GetRenderTargetFn(), Metadata);
 
 				PrepassProxy->DeclareOutputOverride(OutputTextureId);
@@ -479,6 +481,8 @@ bool UCompositeLayerPlate::GetProxy(FTraversalContext& InContext, FSceneRenderin
 		// Register source texture into the frame work context (so that post-processing passes can refer to it).
 		FResourceMetadata Metadata = {};
 		Metadata.bDistorted = true;
+		Metadata.DebugName = TEXT("CompositePlateTex");
+
 		TextureId = InContext.FindOrCreateExternalTexture(Texture, Metadata);
 		
 		const ResourceId OriginalTextureId = TextureId;
diff --git a/Engine/Plugins/Compositing/Composite/Source/Composite/Private/Layers/CompositeLayerSceneCapture.cpp b/Engine/Plugins/Compositing/Composite/Source/Composite/Private/Layers/CompositeLayerSceneCapture.cpp
--- a/Engine/Plugins/Compositing/Composite/Source/Composite/Private/Layers/CompositeLayerSceneCapture.cpp
+++ b/Engine/Plugins/Compositing/Composite/Source/Composite/Private/Layers/CompositeLayerSceneCapture.cpp
@@ -133,6 +133,7 @@ bool UCompositeLayerSceneCapture::GetProxy(FTraversalContext& InContext, FSceneR
 	FResourceMetadata Metadata = {};
 	Metadata.bInvertedAlpha = true;
 	Metadata.bDistorted = CaptureComponent->ShowFlags.LensDistortion;
+	Metadata.DebugName = TEXT("CompositeSceneCaptureTex");
 
 	const ResourceId TexId = InContext.FindOrCreateExternalTexture(CaptureComponent->TextureTarget, Metadata);
 
diff --git a/Engine/Plugins/Compositing/Composite/Source/Composite/Private/Layers/CompositeLayerShadowReflection.cpp b/Engine/Plugins/Compositing/Composite/Source/Composite/Private/Layers/CompositeLayerShadowReflection.cpp
--- a/Engine/Plugins/Compositing/Composite/Source/Composite/Private/Layers/CompositeLayerShadowReflection.cpp
+++ b/Engine/Plugins/Compositing/Composite/Source/Composite/Private/Layers/CompositeLayerShadowReflection.cpp
@@ -143,6 +143,12 @@ bool UCompositeLayerShadowReflection::GetProxy(FTraversalContext& InContext, FSc
 		FPassInputDeclArray Inputs;
 		Inputs.SetNum(2);
 
+		static const TCHAR* DebugNames[] =
+		{
+			TEXT("CompositeShadowReflectionCGTex"),
+			TEXT("CompositeShadowReflectionNoCGTex"),
+		};
+
 		for (int32 Index = 0; Index < 2; ++Index)
 		{
 			UCompositeShadowReflectionCatcherComponent* Component = CachedSceneCaptures[Index].Get();
@@ -151,6 +157,7 @@ bool UCompositeLayerShadowReflection::GetProxy(FTraversalContext& InContext, FSc
 			FResourceMetadata Metadata = {};
 			Metadata.bInvertedAlpha = true;
 			Metadata.bDistorted = Component->ShowFlags.LensDistortion;
+			Metadata.DebugName = DebugNames[Index];
 			const ResourceId TexId = InContext.FindOrCreateExternalTexture(Component->TextureTarget, Metadata);
 
 			Inputs[Index].Set<FPassExternalResourceDesc>({ TexId });
diff --git a/Engine/Plugins/Compositing/Composite/Source/Composite/Private/Layers/CompositeLayerSingleLightShadow.cpp b/Engine/Plugins/Compositing/Composite/Source/Composite/Private/Layers/CompositeLayerSingleLightShadow.cpp
--- a/Engine/Plugins/Compositing/Composite/Source/Composite/Private/Layers/CompositeLayerSingleLightShadow.cpp
+++ b/Engine/Plugins/Compositing/Composite/Source/Composite/Private/Layers/CompositeLayerSingleLightShadow.cpp
@@ -272,10 +272,10 @@ bool UCompositeLayerSingleLightShadow::GetProxy(FTraversalContext& InContext, FS
 	FCompositeSingleLightShadowProxy* ShadowProxy;
 	{
 		// Display-resolution scene depth without jitter as an external texture.
-		const ResourceId SceneDepthId = InContext.FindOrCreateExternalTexture(CachedSceneDepthCapture->TextureTarget, {});
+		const ResourceId SceneDepthId = InContext.FindOrCreateExternalTexture(CachedSceneDepthCapture->TextureTarget, { .DebugName = TEXT("CompositSceneDepthTex") });
 
 		// Shadow map texture
-		const ResourceId ShadowMapId = InContext.FindOrCreateExternalTexture(CachedShadowMapCapture->TextureTarget, {});
+		const ResourceId ShadowMapId = InContext.FindOrCreateExternalTexture(CachedShadowMapCapture->TextureTarget, { .DebugName = TEXT("CompositeShadowMapTex") });
 
 		// Setup shadow pass inputs
 		FPassInputDeclArray PassInputs;
diff --git a/Engine/Plugins/Compositing/CompositeCore/Source/CompositeCore/Private/CompositeCoreSceneViewExtension.cpp b/Engine/Plugins/Compositing/CompositeCore/Source/CompositeCore/Private/CompositeCoreSceneViewExtension.cpp
--- a/Engine/Plugins/Compositing/CompositeCore/Source/CompositeCore/Private/CompositeCoreSceneViewExtension.cpp
+++ b/Engine/Plugins/Compositing/CompositeCore/Source/CompositeCore/Private/CompositeCoreSceneViewExtension.cpp
@@ -159,7 +159,7 @@ void FCompositeCoreSceneViewExtension::SetRenderWork_GameThread(UE::CompositeCor
 							const int32 MapIndex = UE::CompositeCore::EXTERNAL_RANGE_START_ID + Index;
 
 							UE::CompositeCore::FExternalRenderTarget& Target = SVE->ExternalInputs_RenderThread.Add(MapIndex);
-							Target.RenderTarget = CreateRenderTarget(TexResource->GetTextureRHI(), TEXT("CompositeExternalInput"));
+							Target.RenderTarget = CreateRenderTarget(TexResource->GetTextureRHI(), ExternalInput.Metadata.DebugName);
 							Target.Metadata = ExternalInput.Metadata;
 						}
 					}
@@ -554,26 +554,31 @@ UE::CompositeCore::FScopedExternalTextureMap::FScopedExternalTextureMap(const FC
 
 	for (const auto& Pair : SVE.ExternalInputs_RenderThread)
 	{
-		FPassInput& Resource = Textures.Add(Pair.Key);
-		FRDGTextureRef TextureResource = GraphBuilder.RegisterExternalTexture(Pair.Value.RenderTarget);
-		Resource.Texture = FScreenPassTexture{ TextureResource };
-		Resource.Metadata = Pair.Value.Metadata;
-
-		/**
-		 * Mark for internal access and store to reset in destructor.
-		 * 
-		 * Note: we exclude the active view family render target since it is already (and should remain) internal.
-		 * In practice, this can occur when our textures are accessed by pre-processing passes during a scene capture render.
-		 * 
-		 * Unfortunately, full AccessModeState information is not accessible outside of RDG or friend classes. See for example
-		 * GraphBuilder.UseExternalAccessMode(Buffer, AccessModeState.Access, AccessModeState.Pipelines); use in FRDGResourceDumpContext.
-		 */
-		FRHITexture* FamilyRenderTargetRHI = InView.Family->RenderTarget ? InView.Family->RenderTarget->GetRenderTargetTexture().GetReference() : nullptr;
-		if (TextureResource && (TextureResource->GetRHI() != FamilyRenderTargetRHI))
+		const TRefCountPtr<IPooledRenderTarget>& PooledRenderTarget = Pair.Value.RenderTarget;
+		FRDGTextureRef TextureResource = GraphBuilder.RegisterExternalTexture(PooledRenderTarget);
+
+		if (ensureMsgf(TextureResource, TEXT("Failed external pooled render target registration.")))
 		{
-			GraphBuilder.UseInternalAccessMode(TextureResource);
+			FPassInput& Resource = Textures.Add(Pair.Key);
+			Resource.Texture = FScreenPassTexture{ TextureResource };
+			Resource.Metadata = Pair.Value.Metadata;
+
+			/**
+			 * Mark for internal access and store to reset in destructor.
+			 *
+			 * Note: we exclude the active view family render target since it is already (and should remain) internal.
+			 * In practice, this can occur when our textures are accessed by pre-processing passes during a scene capture render.
+			 *
+			 * Unfortunately, full AccessModeState information is not accessible outside of RDG or friend classes. See for example
+			 * GraphBuilder.UseExternalAccessMode(Buffer, AccessModeState.Access, AccessModeState.Pipelines); use in FRDGResourceDumpContext.
+			 */
+			FRHITexture* FamilyRenderTargetRHI = InView.Family->RenderTarget ? InView.Family->RenderTarget->GetRenderTargetTexture().GetReference() : nullptr;
+			if (PooledRenderTarget->GetRHI() != FamilyRenderTargetRHI)
+			{
+				GraphBuilder.UseInternalAccessMode(TextureResource);
 
-			RestoreToExternal.Add(TextureResource);
+				RestoreToExternal.Add(TextureResource);
+			}
 		}
 	}
 }
diff --git a/Engine/Plugins/Compositing/CompositeCore/Source/CompositeCore/Private/Passes/CompositeCorePassMergeProxy.cpp b/Engine/Plugins/Compositing/CompositeCore/Source/CompositeCore/Private/Passes/CompositeCorePassMergeProxy.cpp
--- a/Engine/Plugins/Compositing/CompositeCore/Source/CompositeCore/Private/Passes/CompositeCorePassMergeProxy.cpp
+++ b/Engine/Plugins/Compositing/CompositeCore/Source/CompositeCore/Private/Passes/CompositeCorePassMergeProxy.cpp
@@ -182,7 +182,11 @@ namespace UE
 			if (!Output.IsValid())
 			{
 				const FRDGTextureDesc Desc = Inputs[0].Texture.Texture->Desc;
-				Output = CreateOutputRenderTarget(GraphBuilder, InView, PassContext.OutputViewRect, Desc, TEXT("CompositeCoreMergeOutput"));
+
+				FString& DynamicDebugName = *GraphBuilder.AllocObject<FString>();
+				DynamicDebugName = FString::Printf(TEXT("%sOutput"), ParentName);
+
+				Output = CreateOutputRenderTarget(GraphBuilder, InView, PassContext.OutputViewRect, Desc, *DynamicDebugName);
 			}
 
 			FCompositeCoreMergeShader::FParameters* PassParameters = GraphBuilder.AllocParameters<FCompositeCoreMergeShader::FParameters>();
diff --git a/Engine/Plugins/Compositing/CompositeCore/Source/CompositeCore/Public/Passes/CompositeCorePassMergeProxy.h b/Engine/Plugins/Compositing/CompositeCore/Source/CompositeCore/Public/Passes/CompositeCorePassMergeProxy.h
--- a/Engine/Plugins/Compositing/CompositeCore/Source/CompositeCore/Public/Passes/CompositeCorePassMergeProxy.h
+++ b/Engine/Plugins/Compositing/CompositeCore/Source/CompositeCore/Public/Passes/CompositeCorePassMergeProxy.h
@@ -66,7 +66,7 @@ namespace UE
 			ECompositeCoreMergeOp MergeOp = ECompositeCoreMergeOp::Over;
 
 			/** Parent pass debug name. */
-			const TCHAR* ParentName = TEXT("None");
+			const TCHAR* ParentName = TEXT("MergePass");
 
 			/** Flag to enable or disable application of engine-provided lens distortion LUTs, true by default. */
 			ELensDistortionHandling bUseLensDistortion = ELensDistortionHandling::Enabled;
diff --git a/Engine/Plugins/Compositing/CompositeCore/Source/CompositeCore/Public/Passes/CompositeCorePassProxy.h b/Engine/Plugins/Compositing/CompositeCore/Source/CompositeCore/Public/Passes/CompositeCorePassProxy.h
--- a/Engine/Plugins/Compositing/CompositeCore/Source/CompositeCore/Public/Passes/CompositeCorePassProxy.h
+++ b/Engine/Plugins/Compositing/CompositeCore/Source/CompositeCore/Public/Passes/CompositeCorePassProxy.h
@@ -62,12 +62,20 @@ namespace UE
 			
 			/** Source color encoding. */
 			EEncoding Encoding = EEncoding::Linear;
+
+			/** Debug name */
+			const TCHAR* DebugName = TEXT("CompositeTexture");
 			
 			/** Equality operator */
-			bool operator==(const FResourceMetadata& Other) const = default;
-			
-			/** Inequality operator */
-			bool operator!=(const FResourceMetadata& Other) const = default;
+			bool operator==(const FResourceMetadata& Other) const
+			{
+				return bInvertedAlpha == Other.bInvertedAlpha
+					&& bDistorted == Other.bDistorted
+					&& bPreExposed == Other.bPreExposed
+					&& Encoding == Other.Encoding;
+
+				// DebugName is intentionally ignored
+			}
 		};
 
 		/** Pass texture description for internal resources (default scene textures). */
