Composure: Added support for new masking pass, starting with a mask texture.

* Allowed passes to update the traversal context directly to register external textures.
* In the scene view extension, pre-processing passes now share the same pass traversal code as post-processing passes.
* Fixed various bughawk-flagged issues such as potential division by zero.

#Synced-CL 48811553

diff --git a/Engine/Plugins/Compositing/Composite/Shaders/Private/CompositeCommon.ush b/Engine/Plugins/Compositing/Composite/Shaders/Private/CompositeCommon.ush
--- a/Engine/Plugins/Compositing/Composite/Shaders/Private/CompositeCommon.ush
+++ b/Engine/Plugins/Compositing/Composite/Shaders/Private/CompositeCommon.ush
@@ -2,6 +2,8 @@
 
 #pragma once
 
+#include "/Engine/Private/GammaCorrectionCommon.ush"
+
 #define ESourceEncoding_None  0
 #define ESourceEncoding_Gamma 1
 #define ESourceEncoding_sRGB  2
@@ -22,6 +24,11 @@ float4 Premultiply(float4 InColor)
 	return float4(InColor.xyz * InColor.w, InColor.w);
 }
 
+float InvertAlpha(float InAlpha, uint bInvertedAlpha)
+{
+	return select(bInvertedAlpha > 0, 1.0 - InAlpha, InAlpha);
+}
+
 float3 Decode(float3 InColor, uint InEncoding)
 {
     BRANCH
diff --git a/Engine/Plugins/Compositing/Composite/Shaders/Private/CompositeMasking.usf b/Engine/Plugins/Compositing/Composite/Shaders/Private/CompositeMasking.usf
new file mode 100644
--- /dev/null
+++ b/Engine/Plugins/Compositing/Composite/Shaders/Private/CompositeMasking.usf
@@ -0,0 +1,58 @@
+// Copyright Epic Games, Inc. All Rights Reserved.
+
+#include "/Engine/Private/Common.ush"
+
+#include "CompositeCommon.ush"
+
+#define INPUT_PASS_TEXTURE_MEMBER(StructName, MemberType, MemberName) MemberType StructName##_##MemberName;
+
+#define INPUT_PASS_TEXTURE(StructName) \
+	INPUT_PASS_TEXTURE_MEMBER(StructName, Texture2D, Texture) \
+	INPUT_PASS_TEXTURE_MEMBER(StructName, SamplerState, Sampler) \
+    INPUT_PASS_TEXTURE_MEMBER(StructName, uint, bInvertedAlpha) \
+    INPUT_PASS_TEXTURE_MEMBER(StructName, uint, SourceEncoding)
+
+INPUT_PASS_TEXTURE(Input);
+INPUT_PASS_TEXTURE(Mask);
+
+uint MaskSourceChannel;
+uint bUnpremult;
+uint bInvert;
+
+void MainPS(
+    noperspective float4 InUVAndScreenPos : TEXCOORD0,
+    out float4 OutColor : SV_Target0
+    )
+{
+    float4 InputColor = Input_Texture.SampleLevel(Input_Sampler, InUVAndScreenPos.xy, 0);
+	float4 MaskColor = Mask_Texture.SampleLevel(Mask_Sampler, InUVAndScreenPos.xy, 0);
+
+    // Linearize source colors
+    InputColor.rgb = Decode(InputColor.rgb, Input_SourceEncoding);
+    MaskColor.rgb = Decode(MaskColor.rgb, Mask_SourceEncoding);
+
+	// Normalize alpha as opacity
+	InputColor.a = InvertAlpha(InputColor.a, Input_bInvertedAlpha);
+	MaskColor.a = InvertAlpha(MaskColor.a, Mask_bInvertedAlpha);
+	
+	// Optionally unpremultiply
+	InputColor = select(bUnpremult > 0, Unpremultiply(InputColor), InputColor);
+
+	/** ============ Start Masking ============ */
+	float AlphaMask = saturate(MaskColor[MaskSourceChannel]);
+	AlphaMask = select(bInvert > 0, 1.0f - AlphaMask, AlphaMask);
+	
+	InputColor.a *= AlphaMask;
+	/** ============  End Masking  ============ */
+	
+	// Optionally premultiply
+	InputColor = select(bUnpremult > 0, Premultiply(InputColor), InputColor);
+
+	// Optionally re-invert alpha
+	InputColor.a = InvertAlpha(InputColor.a, Input_bInvertedAlpha);
+
+	// Optionally re-encode
+	InputColor.rgb = Encode(InputColor.rgb, Input_SourceEncoding);
+
+	OutColor = InputColor;
+}
diff --git a/Engine/Plugins/Compositing/Composite/Source/Composite/Private/CompositeActor.cpp b/Engine/Plugins/Compositing/Composite/Source/Composite/Private/CompositeActor.cpp
--- a/Engine/Plugins/Compositing/Composite/Source/Composite/Private/CompositeActor.cpp
+++ b/Engine/Plugins/Compositing/Composite/Source/Composite/Private/CompositeActor.cpp
@@ -95,7 +95,7 @@ namespace UE
 			* 
 			* Returns an index indicating the first proxy whose output can no longer be shared between post-processing locations.
 			*/
-			int32 RegisterLayers(UCompositeLayerBase::FTraversalContext& Context, TArrayView<const TObjectPtr<UCompositeLayerBase>> InLayers, FSceneRenderingBulkObjectAllocator& InFrameAllocator, TArray<FCompositeCorePassProxy*>& OutProxies)
+			int32 RegisterLayers(FCompositeTraversalContext& Context, TArrayView<const TObjectPtr<UCompositeLayerBase>> InLayers, FSceneRenderingBulkObjectAllocator& InFrameAllocator, TArray<FCompositeCorePassProxy*>& OutProxies)
 			{
 				int32 LastLayerIndex = INDEX_NONE;
 				int32 BranchLayerIndex = INDEX_NONE;
@@ -161,7 +161,7 @@ namespace UE
 				}
 
 				TArray<FCompositeCorePassProxy*> LayerProxies;
-				UCompositeLayerBase::FTraversalContext Context;
+				FCompositeTraversalContext Context;
 				Context.bIsSolo = (Layers.Num() == 1);
 
 				const int32 BranchLayerIndex = RegisterLayers(Context, Layers, *RenderWork.FrameAllocator, LayerProxies);
diff --git a/Engine/Plugins/Compositing/Composite/Source/Composite/Private/Layers/CompositeLayerBase.cpp b/Engine/Plugins/Compositing/Composite/Source/Composite/Private/Layers/CompositeLayerBase.cpp
--- a/Engine/Plugins/Compositing/Composite/Source/Composite/Private/Layers/CompositeLayerBase.cpp
+++ b/Engine/Plugins/Compositing/Composite/Source/Composite/Private/Layers/CompositeLayerBase.cpp
@@ -136,7 +136,7 @@ void UCompositeLayerBase::RegisterEndOfFrameUpdate(bool bInEnabled)
 	}
 }
 
-void UCompositeLayerBase::AddChildPasses(UE::CompositeCore::FPassInputDecl& InBasePassInput, FTraversalContext& InContext, FSceneRenderingBulkObjectAllocator& InFrameAllocator, TArrayView<const TObjectPtr<UCompositePassBase>> InPasses) const
+void UCompositeLayerBase::AddChildPasses(UE::CompositeCore::FPassInputDecl& InBasePassInput, FCompositeTraversalContext& InContext, FSceneRenderingBulkObjectAllocator& InFrameAllocator, TArrayView<const TObjectPtr<UCompositePassBase>> InPasses) const
 {
 	// Iterate backwards so that lower passes in the UI are executed first
 	// TODO: Flip back once we have UI customization displaying passes in the correct bottom-up order
@@ -147,15 +147,10 @@ void UCompositeLayerBase::AddChildPasses(UE::CompositeCore::FPassInputDecl& InBa
 		if (IsValid(SubPass) && SubPass->IsActive())
 		{
 			FCompositeCorePassProxy* ChildPassProxy;
-			if (SubPass->GetProxy(InBasePassInput, InFrameAllocator, ChildPassProxy))
+			if (SubPass->GetProxy(InBasePassInput, InContext, InFrameAllocator, ChildPassProxy))
 			{
 				// Next pass should receive the output of the current child pass.
 				InBasePassInput.Set<const FCompositeCorePassProxy*>(ChildPassProxy);
-
-				if (SubPass->NeedsSceneTextures())
-				{
-					InContext.bNeedsSceneTextures = true;
-				}
 			}
 		}
 	}
@@ -173,12 +168,12 @@ const FIntPoint UCompositeLayerBase::GetRenderResolution() const
 	return FCompositeRenderTargetPool::DefaultSize;
 }
 
-ECompositeCoreMergeOp UCompositeLayerBase::GetMergeOperation(const FTraversalContext& InContext) const
+ECompositeCoreMergeOp UCompositeLayerBase::GetMergeOperation(const FCompositeTraversalContext& InContext) const
 {
 	return InContext.bIsSolo ? ECompositeCoreMergeOp::None : Operation;
 }
 
-UE::CompositeCore::FPassInputDecl UCompositeLayerBase::GetDefaultSecondInput(const FTraversalContext& InContext) const
+UE::CompositeCore::FPassInputDecl UCompositeLayerBase::GetDefaultSecondInput(const FCompositeTraversalContext& InContext) const
 {
 	using namespace UE::CompositeCore;
 
@@ -216,26 +211,3 @@ void UCompositeLayerBase::ConditionalOnEndOfFrameUpdate(UWorld* InWorld)
 	}
 }
 
-UE::CompositeCore::ResourceId UCompositeLayerBase::FTraversalContext::FindOrCreateExternalTexture(TWeakObjectPtr<UTexture> InTexture, UE::CompositeCore::FResourceMetadata InMetadata)
-{
-	int32 TextureIndex = ExternalTextures.IndexOfByPredicate([&InTexture, &InMetadata](const UE::CompositeCore::FExternalTexture& InExternalTexture)
-		{
-			return (InExternalTexture.Texture.Get() == InTexture.Get()) && (InExternalTexture.Metadata == InMetadata);
-		}
-	);
-
-	if (TextureIndex == INDEX_NONE)
-	{
-		ExternalTextures.Add(UE::CompositeCore::FExternalTexture{ MoveTemp(InTexture), MoveTemp(InMetadata) });
-		
-		TextureIndex = ExternalTextures.Num() - 1;
-	}
-
-	return UE::CompositeCore::EXTERNAL_RANGE_START_ID + TextureIndex;
-}
-
-const TArray<UE::CompositeCore::FExternalTexture>& UCompositeLayerBase::FTraversalContext::GetExternalTextures() const
-{
-	return ExternalTextures;
-}
-
diff --git a/Engine/Plugins/Compositing/Composite/Source/Composite/Private/Layers/CompositeLayerMainRender.cpp b/Engine/Plugins/Compositing/Composite/Source/Composite/Private/Layers/CompositeLayerMainRender.cpp
--- a/Engine/Plugins/Compositing/Composite/Source/Composite/Private/Layers/CompositeLayerMainRender.cpp
+++ b/Engine/Plugins/Compositing/Composite/Source/Composite/Private/Layers/CompositeLayerMainRender.cpp
@@ -12,7 +12,7 @@ UCompositeLayerMainRender::UCompositeLayerMainRender(const FObjectInitializer& O
 
 UCompositeLayerMainRender::~UCompositeLayerMainRender() = default;
 
-bool UCompositeLayerMainRender::GetProxy(FTraversalContext& InContext, FSceneRenderingBulkObjectAllocator& InFrameAllocator, FCompositeCorePassProxy*& OutProxy) const
+bool UCompositeLayerMainRender::GetProxy(FCompositeTraversalContext& InContext, FSceneRenderingBulkObjectAllocator& InFrameAllocator, FCompositeCorePassProxy*& OutProxy) const
 {
 	using namespace UE::CompositeCore;
 
diff --git a/Engine/Plugins/Compositing/Composite/Source/Composite/Private/Layers/CompositeLayerPlate.cpp b/Engine/Plugins/Compositing/Composite/Source/Composite/Private/Layers/CompositeLayerPlate.cpp
--- a/Engine/Plugins/Compositing/Composite/Source/Composite/Private/Layers/CompositeLayerPlate.cpp
+++ b/Engine/Plugins/Compositing/Composite/Source/Composite/Private/Layers/CompositeLayerPlate.cpp
@@ -407,7 +407,7 @@ int32 UCompositeLayerPlate::FindLastValidPassIndex(TArrayView<const TObjectPtr<U
 }
 
 UE::CompositeCore::ResourceId UCompositeLayerPlate::AddPreprocessingPasses(
-	FTraversalContext& InContext,
+	FCompositeTraversalContext& InContext,
 	FSceneRenderingBulkObjectAllocator& InFrameAllocator,
 	TArrayView<const TObjectPtr<UCompositePassBase>> InPasses,
 	UE::CompositeCore::ResourceId TextureId,
@@ -450,7 +450,7 @@ UE::CompositeCore::ResourceId UCompositeLayerPlate::AddPreprocessingPasses(
 		}
 
 		FCompositeCorePassProxy* PrepassProxy;
-		if (Prepass->GetProxy(PassInput, InFrameAllocator, PrepassProxy))
+		if (Prepass->GetProxy(PassInput, InContext, InFrameAllocator, PrepassProxy))
 		{
 			if (PassIndex == LastValidPassIndex)
 			{
@@ -470,7 +470,7 @@ UE::CompositeCore::ResourceId UCompositeLayerPlate::AddPreprocessingPasses(
 	return OutputTextureId;
 }
 
-bool UCompositeLayerPlate::GetProxy(FTraversalContext& InContext, FSceneRenderingBulkObjectAllocator& InFrameAllocator, FCompositeCorePassProxy*& OutProxy) const
+bool UCompositeLayerPlate::GetProxy(FCompositeTraversalContext& InContext, FSceneRenderingBulkObjectAllocator& InFrameAllocator, FCompositeCorePassProxy*& OutProxy) const
 {
 	using namespace UE::CompositeCore;
 
diff --git a/Engine/Plugins/Compositing/Composite/Source/Composite/Private/Layers/CompositeLayerProcessing.cpp b/Engine/Plugins/Compositing/Composite/Source/Composite/Private/Layers/CompositeLayerProcessing.cpp
--- a/Engine/Plugins/Compositing/Composite/Source/Composite/Private/Layers/CompositeLayerProcessing.cpp
+++ b/Engine/Plugins/Compositing/Composite/Source/Composite/Private/Layers/CompositeLayerProcessing.cpp
@@ -40,7 +40,7 @@ void UCompositeLayerProcessing::PostEditChangeProperty(FPropertyChangedEvent& Pr
 }
 #endif
 
-bool UCompositeLayerProcessing::GetProxy(FTraversalContext& InContext, FSceneRenderingBulkObjectAllocator& InFrameAllocator, FCompositeCorePassProxy*& OutProxy) const
+bool UCompositeLayerProcessing::GetProxy(FCompositeTraversalContext& InContext, FSceneRenderingBulkObjectAllocator& InFrameAllocator, FCompositeCorePassProxy*& OutProxy) const
 {
 	using namespace UE::CompositeCore;
 
@@ -56,17 +56,12 @@ bool UCompositeLayerProcessing::GetProxy(FTraversalContext& InContext, FSceneRen
 		if (IsValid(Pass) && Pass->IsActive())
 		{
 			FCompositeCorePassProxy* PassProxy;
-			if (Pass->GetProxy(BasePassInput, InFrameAllocator, PassProxy))
+			if (Pass->GetProxy(BasePassInput, InContext, InFrameAllocator, PassProxy))
 			{
 				OutProxy = PassProxy;
 
 				// Next pass should receive the output of the current child pass.
 				BasePassInput.Set<const FCompositeCorePassProxy*>(PassProxy);
-
-				if (Pass->NeedsSceneTextures())
-				{
-					InContext.bNeedsSceneTextures = true;
-				}
 			}
 		}
 	}
diff --git a/Engine/Plugins/Compositing/Composite/Source/Composite/Private/Layers/CompositeLayerSceneCapture.cpp b/Engine/Plugins/Compositing/Composite/Source/Composite/Private/Layers/CompositeLayerSceneCapture.cpp
--- a/Engine/Plugins/Compositing/Composite/Source/Composite/Private/Layers/CompositeLayerSceneCapture.cpp
+++ b/Engine/Plugins/Compositing/Composite/Source/Composite/Private/Layers/CompositeLayerSceneCapture.cpp
@@ -112,7 +112,7 @@ void UCompositeLayerSceneCapture::OnEndOfFrameUpdate(UWorld* InWorld)
 	}
 }
 
-bool UCompositeLayerSceneCapture::GetProxy(FTraversalContext& InContext, FSceneRenderingBulkObjectAllocator& InFrameAllocator, FCompositeCorePassProxy*& OutProxy) const
+bool UCompositeLayerSceneCapture::GetProxy(FCompositeTraversalContext& InContext, FSceneRenderingBulkObjectAllocator& InFrameAllocator, FCompositeCorePassProxy*& OutProxy) const
 {
 	using namespace UE::CompositeCore;
 
diff --git a/Engine/Plugins/Compositing/Composite/Source/Composite/Private/Layers/CompositeLayerShadowReflection.cpp b/Engine/Plugins/Compositing/Composite/Source/Composite/Private/Layers/CompositeLayerShadowReflection.cpp
--- a/Engine/Plugins/Compositing/Composite/Source/Composite/Private/Layers/CompositeLayerShadowReflection.cpp
+++ b/Engine/Plugins/Compositing/Composite/Source/Composite/Private/Layers/CompositeLayerShadowReflection.cpp
@@ -111,7 +111,7 @@ void UCompositeLayerShadowReflection::OnEndOfFrameUpdate(UWorld* InWorld)
 	}
 }
 
-bool UCompositeLayerShadowReflection::GetProxy(FTraversalContext& InContext, FSceneRenderingBulkObjectAllocator& InFrameAllocator, FCompositeCorePassProxy*& OutProxy) const
+bool UCompositeLayerShadowReflection::GetProxy(FCompositeTraversalContext& InContext, FSceneRenderingBulkObjectAllocator& InFrameAllocator, FCompositeCorePassProxy*& OutProxy) const
 {
 	using namespace UE::CompositeCore;
 
diff --git a/Engine/Plugins/Compositing/Composite/Source/Composite/Private/Layers/CompositeLayerSingleLightShadow.cpp b/Engine/Plugins/Compositing/Composite/Source/Composite/Private/Layers/CompositeLayerSingleLightShadow.cpp
--- a/Engine/Plugins/Compositing/Composite/Source/Composite/Private/Layers/CompositeLayerSingleLightShadow.cpp
+++ b/Engine/Plugins/Compositing/Composite/Source/Composite/Private/Layers/CompositeLayerSingleLightShadow.cpp
@@ -97,7 +97,7 @@ namespace UE
 				Parameters->SceneDepthScreenToWorld		= SceneDepthScreenToWorld;
 				Parameters->ShadowInvDeviceZToWorldZ	= ShadowInvDeviceZToWorldZ;
 				Parameters->ShadowMatrix				= ShadowMatrix;
-				Parameters->ShadowBufferSize			= FVector4f((float)ShadowMapSize.X, (float)ShadowMapSize.Y, 1.0f / (float)ShadowMapSize.X, 1.0f / (float)ShadowMapSize.Y);
+				Parameters->ShadowBufferSize			= FVector4f((float)ShadowMapSize.X, (float)ShadowMapSize.Y, FMath::SafeDivide(1.0f, (float)ShadowMapSize.X), FMath::SafeDivide(1.0f, (float)ShadowMapSize.Y));
 				Parameters->ShadowStrengthBias		= ShadowStrengthBias;
 				Parameters->LightWorldDirection			= LightWorldDirection;
 				// We can use the current view since we the depth render matches the main view
@@ -193,7 +193,7 @@ void UCompositeLayerSingleLightShadow::OnEndOfFrameUpdate(UWorld* InWorld)
 	}
 }
 
-bool UCompositeLayerSingleLightShadow::GetProxy(FTraversalContext& InContext, FSceneRenderingBulkObjectAllocator& InFrameAllocator, FCompositeCorePassProxy*& OutProxy) const
+bool UCompositeLayerSingleLightShadow::GetProxy(FCompositeTraversalContext& InContext, FSceneRenderingBulkObjectAllocator& InFrameAllocator, FCompositeCorePassProxy*& OutProxy) const
 {
 	using namespace UE::CompositeCore;
 
diff --git a/Engine/Plugins/Compositing/Composite/Source/Composite/Private/Passes/CompositePassBase.cpp b/Engine/Plugins/Compositing/Composite/Source/Composite/Private/Passes/CompositePassBase.cpp
--- a/Engine/Plugins/Compositing/Composite/Source/Composite/Private/Passes/CompositePassBase.cpp
+++ b/Engine/Plugins/Compositing/Composite/Source/Composite/Private/Passes/CompositePassBase.cpp
@@ -4,6 +4,29 @@
 
 #include "CompositeActor.h"
 
+UE::CompositeCore::ResourceId FCompositeTraversalContext::FindOrCreateExternalTexture(TWeakObjectPtr<UTexture> InTexture, UE::CompositeCore::FResourceMetadata InMetadata)
+{
+	int32 TextureIndex = ExternalTextures.IndexOfByPredicate([&InTexture, &InMetadata](const UE::CompositeCore::FExternalTexture& InExternalTexture)
+		{
+			return (InExternalTexture.Texture.Get() == InTexture.Get()) && (InExternalTexture.Metadata == InMetadata);
+		}
+	);
+
+	if (TextureIndex == INDEX_NONE)
+	{
+		ExternalTextures.Add(UE::CompositeCore::FExternalTexture{ MoveTemp(InTexture), MoveTemp(InMetadata) });
+
+		TextureIndex = ExternalTextures.Num() - 1;
+	}
+
+	return UE::CompositeCore::EXTERNAL_RANGE_START_ID + TextureIndex;
+}
+
+const TArray<UE::CompositeCore::FExternalTexture>& FCompositeTraversalContext::GetExternalTextures() const
+{
+	return ExternalTextures;
+}
+
 UCompositePassBase::UCompositePassBase(const FObjectInitializer& ObjectInitializer)
 	: Super(ObjectInitializer)
 {
diff --git a/Engine/Plugins/Compositing/Composite/Source/Composite/Private/Passes/CompositePassCenteredScale.cpp b/Engine/Plugins/Compositing/Composite/Source/Composite/Private/Passes/CompositePassCenteredScale.cpp
--- a/Engine/Plugins/Compositing/Composite/Source/Composite/Private/Passes/CompositePassCenteredScale.cpp
+++ b/Engine/Plugins/Compositing/Composite/Source/Composite/Private/Passes/CompositePassCenteredScale.cpp
@@ -95,7 +95,7 @@ UCompositePassCenteredScale::UCompositePassCenteredScale(const FObjectInitialize
 
 UCompositePassCenteredScale::~UCompositePassCenteredScale() = default;
 
-bool UCompositePassCenteredScale::GetProxy(const UE::CompositeCore::FPassInputDecl& InputDecl, FSceneRenderingBulkObjectAllocator& InFrameAllocator, FCompositeCorePassProxy*& OutProxy) const
+bool UCompositePassCenteredScale::GetProxy(const UE::CompositeCore::FPassInputDecl& InputDecl, FCompositeTraversalContext& InContext, FSceneRenderingBulkObjectAllocator& InFrameAllocator, FCompositeCorePassProxy*& OutProxy) const
 {
 	using namespace UE::CompositeCore;
 
@@ -161,7 +161,7 @@ FVector2f UCompositePassCenteredScale::CalculateScale() const
 				}
 				else if (const UTexture2D* Texture2D = Cast<UTexture2D>(PlateLayer->Texture))
 				{
-					SourceAR = Texture2D->GetSizeX() / static_cast<float>(Texture2D->GetSizeY());
+					SourceAR = FMath::SafeDivide(Texture2D->GetSizeX(), static_cast<float>(Texture2D->GetSizeY()));
 				}
 			}
 		}
@@ -173,14 +173,15 @@ FVector2f UCompositePassCenteredScale::CalculateScale() const
 	break;
 	case ECompositePassScaleMode::AspectRatio:
 	{
-		const float SourceAR = SourceAspectRatio.X / SourceAspectRatio.Y;
-		const float TargetAR = TargetAspectRatio.X / TargetAspectRatio.Y;
+		const float SourceAR = FMath::SafeDivide(SourceAspectRatio.X, SourceAspectRatio.Y);
+		const float TargetAR = FMath::SafeDivide(TargetAspectRatio.X, TargetAspectRatio.Y);
 
 		ApplyAspectRatioScaleFn(SourceAR, TargetAR);
 	}
 	break;
 	case ECompositePassScaleMode::Manual:
-		ScaleUV /= ScaleFactor;
+		ScaleUV.X = FMath::SafeDivide(ScaleUV.X, ScaleFactor.X);
+		ScaleUV.Y = FMath::SafeDivide(ScaleUV.Y, ScaleFactor.Y);
 		break;
 
 	default:
diff --git a/Engine/Plugins/Compositing/Composite/Source/Composite/Private/Passes/CompositePassColorGrading.cpp b/Engine/Plugins/Compositing/Composite/Source/Composite/Private/Passes/CompositePassColorGrading.cpp
--- a/Engine/Plugins/Compositing/Composite/Source/Composite/Private/Passes/CompositePassColorGrading.cpp
+++ b/Engine/Plugins/Compositing/Composite/Source/Composite/Private/Passes/CompositePassColorGrading.cpp
@@ -129,7 +129,7 @@ UCompositePassColorGrading::UCompositePassColorGrading(const FObjectInitializer&
 
 UCompositePassColorGrading::~UCompositePassColorGrading() = default;
 
-bool UCompositePassColorGrading::GetProxy(const UE::CompositeCore::FPassInputDecl& InputDecl, FSceneRenderingBulkObjectAllocator& InFrameAllocator, FCompositeCorePassProxy*& OutProxy) const
+bool UCompositePassColorGrading::GetProxy(const UE::CompositeCore::FPassInputDecl& InputDecl, FCompositeTraversalContext& InContext, FSceneRenderingBulkObjectAllocator& InFrameAllocator, FCompositeCorePassProxy*& OutProxy) const
 {
 	FColorGradingCompositePassProxy* Proxy = InFrameAllocator.Create<FColorGradingCompositePassProxy>(UE::CompositeCore::FPassInputDeclArray{ InputDecl });
 	Proxy->TemperatureSettings = TemperatureSettings;
diff --git a/Engine/Plugins/Compositing/Composite/Source/Composite/Private/Passes/CompositePassColorKeyer.cpp b/Engine/Plugins/Compositing/Composite/Source/Composite/Private/Passes/CompositePassColorKeyer.cpp
--- a/Engine/Plugins/Compositing/Composite/Source/Composite/Private/Passes/CompositePassColorKeyer.cpp
+++ b/Engine/Plugins/Compositing/Composite/Source/Composite/Private/Passes/CompositePassColorKeyer.cpp
@@ -232,7 +232,7 @@ UCompositePassColorKeyer::UCompositePassColorKeyer(const FObjectInitializer& Obj
 
 UCompositePassColorKeyer::~UCompositePassColorKeyer() = default;
 
-bool UCompositePassColorKeyer::GetProxy(const UE::CompositeCore::FPassInputDecl& InputDecl, FSceneRenderingBulkObjectAllocator& InFrameAllocator, FCompositeCorePassProxy*& OutProxy) const
+bool UCompositePassColorKeyer::GetProxy(const UE::CompositeCore::FPassInputDecl& InputDecl, FCompositeTraversalContext& InContext, FSceneRenderingBulkObjectAllocator& InFrameAllocator, FCompositeCorePassProxy*& OutProxy) const
 {
 	using namespace UE::CompositeCore;
 
diff --git a/Engine/Plugins/Compositing/Composite/Source/Composite/Private/Passes/CompositePassDistortion.cpp b/Engine/Plugins/Compositing/Composite/Source/Composite/Private/Passes/CompositePassDistortion.cpp
--- a/Engine/Plugins/Compositing/Composite/Source/Composite/Private/Passes/CompositePassDistortion.cpp
+++ b/Engine/Plugins/Compositing/Composite/Source/Composite/Private/Passes/CompositePassDistortion.cpp
@@ -196,7 +196,7 @@ bool UCompositePassDistortion::IsActive() const
 	return bIsDistortionActive;
 }
 
-bool UCompositePassDistortion::GetProxy(const UE::CompositeCore::FPassInputDecl& InputDecl, FSceneRenderingBulkObjectAllocator& InFrameAllocator, FCompositeCorePassProxy*& OutProxy) const
+bool UCompositePassDistortion::GetProxy(const UE::CompositeCore::FPassInputDecl& InputDecl, FCompositeTraversalContext& InContext, FSceneRenderingBulkObjectAllocator& InFrameAllocator, FCompositeCorePassProxy*& OutProxy) const
 {
 	if (IsActive())
 	{
diff --git a/Engine/Plugins/Compositing/Composite/Source/Composite/Private/Passes/CompositePassFXAA.cpp b/Engine/Plugins/Compositing/Composite/Source/Composite/Private/Passes/CompositePassFXAA.cpp
--- a/Engine/Plugins/Compositing/Composite/Source/Composite/Private/Passes/CompositePassFXAA.cpp
+++ b/Engine/Plugins/Compositing/Composite/Source/Composite/Private/Passes/CompositePassFXAA.cpp
@@ -12,7 +12,7 @@ UCompositePassFXAA::UCompositePassFXAA(const FObjectInitializer& ObjectInitializ
 
 UCompositePassFXAA::~UCompositePassFXAA() = default;
 
-bool UCompositePassFXAA::GetProxy(const UE::CompositeCore::FPassInputDecl& InputDecl, FSceneRenderingBulkObjectAllocator& InFrameAllocator, FCompositeCorePassProxy*& OutProxy) const
+bool UCompositePassFXAA::GetProxy(const UE::CompositeCore::FPassInputDecl& InputDecl, FCompositeTraversalContext& InContext, FSceneRenderingBulkObjectAllocator& InFrameAllocator, FCompositeCorePassProxy*& OutProxy) const
 {
 	using namespace UE::CompositeCore;
 
diff --git a/Engine/Plugins/Compositing/Composite/Source/Composite/Private/Passes/CompositePassMasking.cpp b/Engine/Plugins/Compositing/Composite/Source/Composite/Private/Passes/CompositePassMasking.cpp
new file mode 100644
--- /dev/null
+++ b/Engine/Plugins/Compositing/Composite/Source/Composite/Private/Passes/CompositePassMasking.cpp
@@ -0,0 +1,149 @@
+// Copyright Epic Games, Inc. All Rights Reserved.
+
+#include "Passes/CompositePassMasking.h"
+#include "Passes/CompositeCorePassProxy.h"
+
+#include "Engine/Texture.h"
+#include "RenderGraphBuilder.h"
+#include "RenderGraphUtils.h"
+#include "SceneView.h"
+#include "ScreenPass.h"
+
+DECLARE_GPU_STAT_NAMED(FCompositeMasking, TEXT("Composite.Masking"));
+
+class FCompositePassMaskingShader : public FGlobalShader
+{
+	DECLARE_GLOBAL_SHADER(FCompositePassMaskingShader);
+	SHADER_USE_PARAMETER_STRUCT(FCompositePassMaskingShader, FGlobalShader);
+
+	BEGIN_SHADER_PARAMETER_STRUCT(FInputPassTextureParameters, )
+		SHADER_PARAMETER_RDG_TEXTURE(Texture2D, Texture)
+		SHADER_PARAMETER_SAMPLER(SamplerState, Sampler)
+		SHADER_PARAMETER(uint32, bInvertedAlpha)
+		SHADER_PARAMETER(uint32, SourceEncoding)
+	END_SHADER_PARAMETER_STRUCT()
+
+	BEGIN_SHADER_PARAMETER_STRUCT(FParameters, )
+		SHADER_PARAMETER_STRUCT(FInputPassTextureParameters, Input)
+		SHADER_PARAMETER_STRUCT(FInputPassTextureParameters, Mask)
+		SHADER_PARAMETER(uint32, MaskSourceChannel)
+		SHADER_PARAMETER(uint32, bUnpremult)
+		SHADER_PARAMETER(uint32, bInvert)
+		RENDER_TARGET_BINDING_SLOTS()
+	END_SHADER_PARAMETER_STRUCT()
+};
+IMPLEMENT_GLOBAL_SHADER(FCompositePassMaskingShader, "/Plugin/Composite/Private/CompositeMasking.usf", "MainPS", SF_Pixel);
+
+
+namespace UE
+{
+	namespace Composite
+	{
+		using namespace CompositeCore;
+
+		class FMaskingPassProxy : public FCompositeCorePassProxy
+		{
+		public:
+			IMPLEMENT_COMPOSITE_PASS(FMaskingPassProxy);
+
+			using FCompositeCorePassProxy::FCompositeCorePassProxy;
+
+			FCompositePassMaskingShader::FInputPassTextureParameters GetPassTextureParameters(const UE::CompositeCore::FPassInput& Input) const
+			{
+				FCompositePassMaskingShader::FInputPassTextureParameters Parameters;
+				Parameters.Texture = Input.Texture.Texture;
+				Parameters.Sampler = TStaticSamplerState<SF_Bilinear>::GetRHI();
+				Parameters.bInvertedAlpha = Input.Metadata.bInvertedAlpha;
+				Parameters.SourceEncoding = static_cast<uint32>(Input.Metadata.Encoding);
+
+				return Parameters;
+			}
+
+			FPassTexture Add(FRDGBuilder& GraphBuilder, const FSceneView& InView, const FPassInputArray& Inputs, const FPassContext& PassContext) const override
+			{
+				RDG_EVENT_SCOPE_STAT(GraphBuilder, FCompositeMasking, "Composite.Masking");
+				RDG_GPU_STAT_SCOPE(GraphBuilder, FCompositeMasking);
+
+				check(ValidateInputs(Inputs));
+
+				const UE::CompositeCore::FResourceMetadata& Metadata = Inputs[0].Metadata;
+				const FScreenPassTexture& InputTexture = Inputs[0].Texture;
+				FScreenPassRenderTarget Output = Inputs.OverrideOutput;
+
+				// If the override output is provided, it means that this is the last pass in post processing.
+				if (!Output.IsValid())
+				{
+					Output = CreateOutputRenderTarget(GraphBuilder, InView, PassContext.OutputViewRect, InputTexture.Texture->Desc, TEXT("MaskingCompositePass"));
+				}
+
+				FGlobalShaderMap* GlobalShaderMap = GetGlobalShaderMap(InView.GetFeatureLevel());
+				FScreenPassTextureViewport Viewport = FScreenPassTextureViewport(Output);
+
+				FCompositePassMaskingShader::FParameters* Parameters = GraphBuilder.AllocParameters<FCompositePassMaskingShader::FParameters>();
+				Parameters->Input = GetPassTextureParameters(Inputs[0]);
+				Parameters->Mask = GetPassTextureParameters(Inputs[1]);
+				Parameters->MaskSourceChannel = MaskSourceChannel;
+				Parameters->bUnpremult = bUnpremult;
+				Parameters->bInvert = bInvert;
+				Parameters->RenderTargets[0] = Output.GetRenderTargetBinding();
+
+				TShaderMapRef<FCompositePassMaskingShader> PixelShader(GlobalShaderMap);
+				AddDrawScreenPass(
+					GraphBuilder,
+					RDG_EVENT_NAME("Composite.Masking (%dx%d)", Viewport.Extent.X, Viewport.Extent.Y),
+					InView,
+					Viewport,
+					Viewport,
+					PixelShader,
+					Parameters
+				);
+
+				return UE::CompositeCore::FPassTexture{ MoveTemp(Output), Metadata };
+			}
+
+			uint32 MaskSourceChannel = 0;
+			bool bUnpremult = true;
+			bool bInvert = false;
+		};
+	}
+}
+
+UCompositePassMasking::UCompositePassMasking(const FObjectInitializer& ObjectInitializer)
+	: Super(ObjectInitializer)
+	, MaskTexture{}
+	, MaskSource{ ECompositeColorChannel::Red }
+	, bUnpremult{ true }
+	, bInvert{ false }
+{
+}
+
+UCompositePassMasking::~UCompositePassMasking() = default;
+
+bool UCompositePassMasking::IsActive() const
+{
+	return IsEnabled() && IsValid(MaskTexture);
+}
+
+bool UCompositePassMasking::GetProxy(const UE::CompositeCore::FPassInputDecl& InputDecl, FCompositeTraversalContext& InContext, FSceneRenderingBulkObjectAllocator& InFrameAllocator, FCompositeCorePassProxy*& OutProxy) const
+{
+	using namespace UE::Composite;
+	const bool bValidTexture = IsValid(MaskTexture);
+
+	if (ensureMsgf(bValidTexture, TEXT("Texture is expected to be valid for active pass.")))
+	{
+		ResourceId TextureId = InContext.FindOrCreateExternalTexture(MaskTexture, { .DebugName = TEXT("MaskTexture") });
+
+		FPassInputDecl MaskingTextureInput;
+		MaskingTextureInput.Set<FPassExternalResourceDesc>({ TextureId });
+
+		FPassInputDeclArray Inputs = { InputDecl, MaskingTextureInput };
+		FMaskingPassProxy* Proxy = InFrameAllocator.Create<FMaskingPassProxy>(Inputs);
+		Proxy->MaskSourceChannel = static_cast<uint32>(MaskSource);
+		Proxy->bUnpremult = bUnpremult;
+		Proxy->bInvert = bInvert;
+
+		OutProxy = Proxy;
+	}
+
+	return bValidTexture;
+}
diff --git a/Engine/Plugins/Compositing/Composite/Source/Composite/Private/Passes/CompositePassMaterial.cpp b/Engine/Plugins/Compositing/Composite/Source/Composite/Private/Passes/CompositePassMaterial.cpp
--- a/Engine/Plugins/Compositing/Composite/Source/Composite/Private/Passes/CompositePassMaterial.cpp
+++ b/Engine/Plugins/Compositing/Composite/Source/Composite/Private/Passes/CompositePassMaterial.cpp
@@ -60,12 +60,16 @@ bool UCompositePassMaterial::IsActive() const
 	return Super::IsActive() && IsValid(PostProcessMaterial) && PostProcessMaterial->IsPostProcessMaterial();
 }
 
-bool UCompositePassMaterial::GetProxy(const UE::CompositeCore::FPassInputDecl& InputDecl, FSceneRenderingBulkObjectAllocator& InFrameAllocator, FCompositeCorePassProxy*& OutProxy) const
+bool UCompositePassMaterial::GetProxy(const UE::CompositeCore::FPassInputDecl& InputDecl, FCompositeTraversalContext& InContext, FSceneRenderingBulkObjectAllocator& InFrameAllocator, FCompositeCorePassProxy*& OutProxy) const
 {
 	if (IsActive())
 	{
 		FCompositePassMaterialProxy* Proxy = InFrameAllocator.Create<FCompositePassMaterialProxy>(UE::CompositeCore::FPassInputDeclArray{ InputDecl });
 		Proxy->Material = PostProcessMaterial;
+
+		// We conservatively mark any post-processing materials as potentially using scene textures.
+		// TBD if material inputs could be individually checked instead.
+		InContext.bNeedsSceneTextures = true;
 		
 		OutProxy = Proxy;
 		return true;
@@ -74,10 +78,3 @@ bool UCompositePassMaterial::GetProxy(const UE::CompositeCore::FPassInputDecl& I
 	return false;
 }
 
-bool UCompositePassMaterial::NeedsSceneTextures() const
-{
-	// We conservatively mark any post-processing materials as potentially using scene textures.
-	// TBD if material inputs could be individually checked instead.
-	return true;
-}
-
diff --git a/Engine/Plugins/Compositing/Composite/Source/Composite/Private/Passes/CompositePassOpenColorIO.cpp b/Engine/Plugins/Compositing/Composite/Source/Composite/Private/Passes/CompositePassOpenColorIO.cpp
--- a/Engine/Plugins/Compositing/Composite/Source/Composite/Private/Passes/CompositePassOpenColorIO.cpp
+++ b/Engine/Plugins/Compositing/Composite/Source/Composite/Private/Passes/CompositePassOpenColorIO.cpp
@@ -60,7 +60,7 @@ UCompositePassOpenColorIO::UCompositePassOpenColorIO(const FObjectInitializer& O
 
 UCompositePassOpenColorIO::~UCompositePassOpenColorIO() = default;
 
-bool UCompositePassOpenColorIO::GetProxy(const UE::CompositeCore::FPassInputDecl& InputDecl, FSceneRenderingBulkObjectAllocator& InFrameAllocator, FCompositeCorePassProxy*& OutProxy) const
+bool UCompositePassOpenColorIO::GetProxy(const UE::CompositeCore::FPassInputDecl& InputDecl, FCompositeTraversalContext& InContext, FSceneRenderingBulkObjectAllocator& InFrameAllocator, FCompositeCorePassProxy*& OutProxy) const
 {
 	FOpenColorIOCompositePassProxy* Proxy = InFrameAllocator.Create<FOpenColorIOCompositePassProxy>(UE::CompositeCore::FPassInputDeclArray{ InputDecl },
 		FOpenColorIORendering::GetRenderPassResources(ColorConversionSettings, GMaxRHIFeatureLevel)
diff --git a/Engine/Plugins/Compositing/Composite/Source/Composite/Private/Passes/CompositePassSMAA.cpp b/Engine/Plugins/Compositing/Composite/Source/Composite/Private/Passes/CompositePassSMAA.cpp
--- a/Engine/Plugins/Compositing/Composite/Source/Composite/Private/Passes/CompositePassSMAA.cpp
+++ b/Engine/Plugins/Compositing/Composite/Source/Composite/Private/Passes/CompositePassSMAA.cpp
@@ -44,7 +44,8 @@ namespace UE
 
 				static IConsoleVariable* SMAAQualityCVar = IConsoleManager::Get().FindConsoleVariable(TEXT("r.SMAA.Quality"));
 				check(SMAAQualityCVar);
-				const int32 Quality = QualityOverride.IsSet() ? QualityOverride.GetValue() : SMAAQualityCVar->GetInt();
+				int32 Quality = SMAAQualityCVar ? SMAAQualityCVar->GetInt() : static_cast<int32>(ESMAAQuality::MAX);
+				Quality = QualityOverride.Get(Quality);
 
 				FSMAAInputs SMAAInputs{};
 				SMAAInputs.SceneColor = Input;
@@ -81,7 +82,7 @@ UCompositePassSMAA::UCompositePassSMAA(const FObjectInitializer& ObjectInitializ
 
 UCompositePassSMAA::~UCompositePassSMAA() = default;
 
-bool UCompositePassSMAA::GetProxy(const UE::CompositeCore::FPassInputDecl& InputDecl, FSceneRenderingBulkObjectAllocator& InFrameAllocator, FCompositeCorePassProxy*& OutProxy) const
+bool UCompositePassSMAA::GetProxy(const UE::CompositeCore::FPassInputDecl& InputDecl, FCompositeTraversalContext& InContext, FSceneRenderingBulkObjectAllocator& InFrameAllocator, FCompositeCorePassProxy*& OutProxy) const
 {
 	using namespace UE::Composite;
 
diff --git a/Engine/Plugins/Compositing/Composite/Source/Composite/Public/Layers/CompositeLayerBase.h b/Engine/Plugins/Compositing/Composite/Source/Composite/Public/Layers/CompositeLayerBase.h
--- a/Engine/Plugins/Compositing/Composite/Source/Composite/Public/Layers/CompositeLayerBase.h
+++ b/Engine/Plugins/Compositing/Composite/Source/Composite/Public/Layers/CompositeLayerBase.h
@@ -23,36 +23,6 @@ class UCompositeLayerBase : public UObject
 	GENERATED_BODY()
 
 public:
-	/** Context used during layer to proxy assembly. */
-	struct FTraversalContext
-	{
-		/** Flag to indicate solo pass */
-		bool bIsSolo = false;
-		
-		/** Flag to indicate first pass */
-		bool bIsFirstPass = true;
-
-		/** Flag to indicate that a layer is independent of scene textures. Work from scene-independent layers can be shared between post-processing locations. */
-		bool bNeedsSceneTextures = false;
-		
-		/** Array of passes applied at the start of rendering. */
-		TSortedMap<UE::CompositeCore::ResourceId, TArray<const FCompositeCorePassProxy*>> PreprocessingPasses;
-
-		/**
-		* Find or create an external texture resource.
-		* 
-		* @return ResourceId Texture resource identifier
-		*/
-		UE::CompositeCore::ResourceId FindOrCreateExternalTexture(TWeakObjectPtr<UTexture> InTexture, UE::CompositeCore::FResourceMetadata InMetadata);
-
-		/** Get external textures. */
-		const TArray<UE::CompositeCore::FExternalTexture>& GetExternalTextures() const;
-
-	private:
-		/** List of external textures */
-		TArray<UE::CompositeCore::FExternalTexture> ExternalTextures;
-	};
-
 	/** Constructor. */
 	UE_API UCompositeLayerBase(const FObjectInitializer& ObjectInitializer);
 	
@@ -89,7 +59,7 @@ public:
 	UE_API virtual void OnRemoved(ACompositeActor* LastOwner) {}
 
 	/** Getter function to override, returning pass layer proxies to be passed to the render thread. (Proxy objects should be allocated from the provided allocator.) */
-	UE_API virtual bool GetProxy(FTraversalContext& InContext, FSceneRenderingBulkObjectAllocator& InFrameAllocator, FCompositeCorePassProxy*& OutProxy) const { return false; };
+	UE_API virtual bool GetProxy(FCompositeTraversalContext& InContext, FSceneRenderingBulkObjectAllocator& InFrameAllocator, FCompositeCorePassProxy*& OutProxy) const { return false; };
 
 	/** Get the enabled state. */
 	UFUNCTION(BlueprintGetter)
@@ -127,7 +97,7 @@ protected:
 	bool bIsEnabled = true;
 
 	/** Convenience function to register post-processing child passes. */
-	void AddChildPasses(UE::CompositeCore::FPassInputDecl& InBasePassInput, FTraversalContext& InContext, FSceneRenderingBulkObjectAllocator& InFrameAllocator, TArrayView<const TObjectPtr<UCompositePassBase>> InPasses) const;
+	void AddChildPasses(UE::CompositeCore::FPassInputDecl& InBasePassInput, FCompositeTraversalContext& InContext, FSceneRenderingBulkObjectAllocator& InFrameAllocator, TArrayView<const TObjectPtr<UCompositePassBase>> InPasses) const;
 
 	/** Fixed number of inputs used by layer merge passes. */
 	static constexpr int32 FixedNumLayerInputs = 2;
@@ -136,10 +106,10 @@ protected:
 	const FIntPoint GetRenderResolution() const;
 
 	/** Convenience function that returns the merge operation depending on the traversal context. */
-	ECompositeCoreMergeOp GetMergeOperation(const FTraversalContext& InContext) const;
+	ECompositeCoreMergeOp GetMergeOperation(const FCompositeTraversalContext& InContext) const;
 
 	/** Convenience function that returns the default secondary input depending on the traversal context. */
-	UE::CompositeCore::FPassInputDecl GetDefaultSecondInput(const FTraversalContext& InContext) const;
+	UE::CompositeCore::FPassInputDecl GetDefaultSecondInput(const FCompositeTraversalContext& InContext) const;
 
 private:
 	/** Private end of frame callback implementation to condition the inherited function. */
diff --git a/Engine/Plugins/Compositing/Composite/Source/Composite/Public/Layers/CompositeLayerMainRender.h b/Engine/Plugins/Compositing/Composite/Source/Composite/Public/Layers/CompositeLayerMainRender.h
--- a/Engine/Plugins/Compositing/Composite/Source/Composite/Public/Layers/CompositeLayerMainRender.h
+++ b/Engine/Plugins/Compositing/Composite/Source/Composite/Public/Layers/CompositeLayerMainRender.h
@@ -24,7 +24,7 @@ public:
 	UE_API virtual ~UCompositeLayerMainRender();
 
 	/** Getter function to override, returning pass layer proxies to be passed to the render thread. (Proxy objects should be allocated from the provided allocator.) */
-	UE_API virtual bool GetProxy(FTraversalContext& InContext, FSceneRenderingBulkObjectAllocator& InFrameAllocator, FCompositeCorePassProxy*& OutProxy) const override;
+	UE_API virtual bool GetProxy(FCompositeTraversalContext& InContext, FSceneRenderingBulkObjectAllocator& InFrameAllocator, FCompositeCorePassProxy*& OutProxy) const override;
 };
 
 #undef UE_API
diff --git a/Engine/Plugins/Compositing/Composite/Source/Composite/Public/Layers/CompositeLayerPlate.h b/Engine/Plugins/Compositing/Composite/Source/Composite/Public/Layers/CompositeLayerPlate.h
--- a/Engine/Plugins/Compositing/Composite/Source/Composite/Public/Layers/CompositeLayerPlate.h
+++ b/Engine/Plugins/Compositing/Composite/Source/Composite/Public/Layers/CompositeLayerPlate.h
@@ -63,7 +63,7 @@ public:
 	//~ End FTickableGameObject interface
 
 	/** Getter function to override, returning pass layer proxies to be passed to the render thread. (Proxy objects should be allocated from the provided allocator.) */
-	UE_API virtual bool GetProxy(FTraversalContext& InContext, FSceneRenderingBulkObjectAllocator& InFrameAllocator, FCompositeCorePassProxy*& OutProxy) const override;
+	UE_API virtual bool GetProxy(FCompositeTraversalContext& InContext, FSceneRenderingBulkObjectAllocator& InFrameAllocator, FCompositeCorePassProxy*& OutProxy) const override;
 
 	/** Get array of all composite mesh primitives. */
 	UE_API TArray<UPrimitiveComponent*> GetPrimitives() const;
@@ -184,7 +184,7 @@ private:
 
 	/** Convenience function to register pre-processing child passes. */
 	UE::CompositeCore::ResourceId AddPreprocessingPasses(
-		FTraversalContext& InContext,
+		FCompositeTraversalContext& InContext,
 		FSceneRenderingBulkObjectAllocator& InFrameAllocator,
 		TArrayView<const TObjectPtr<UCompositePassBase>> InPasses,
 		UE::CompositeCore::ResourceId TextureId,
diff --git a/Engine/Plugins/Compositing/Composite/Source/Composite/Public/Layers/CompositeLayerProcessing.h b/Engine/Plugins/Compositing/Composite/Source/Composite/Public/Layers/CompositeLayerProcessing.h
--- a/Engine/Plugins/Compositing/Composite/Source/Composite/Public/Layers/CompositeLayerProcessing.h
+++ b/Engine/Plugins/Compositing/Composite/Source/Composite/Public/Layers/CompositeLayerProcessing.h
@@ -24,7 +24,7 @@ public:
 	UE_API virtual ~UCompositeLayerProcessing();
 
 	/** Getter function to override, returning pass layer proxies to be passed to the render thread. (Proxy objects should be allocated from the provided allocator.) */
-	UE_API virtual bool GetProxy(FTraversalContext& InContext, FSceneRenderingBulkObjectAllocator& InFrameAllocator, FCompositeCorePassProxy*& OutProxy) const override;
+	UE_API virtual bool GetProxy(FCompositeTraversalContext& InContext, FSceneRenderingBulkObjectAllocator& InFrameAllocator, FCompositeCorePassProxy*& OutProxy) const override;
 
 #if WITH_EDITOR
 	UE_API virtual bool CanEditChange(const FProperty* InProperty) const;
diff --git a/Engine/Plugins/Compositing/Composite/Source/Composite/Public/Layers/CompositeLayerSceneCapture.h b/Engine/Plugins/Compositing/Composite/Source/Composite/Public/Layers/CompositeLayerSceneCapture.h
--- a/Engine/Plugins/Compositing/Composite/Source/Composite/Public/Layers/CompositeLayerSceneCapture.h
+++ b/Engine/Plugins/Compositing/Composite/Source/Composite/Public/Layers/CompositeLayerSceneCapture.h
@@ -44,7 +44,7 @@ public:
 	UE_API virtual void OnEndOfFrameUpdate(UWorld* InWorld) override;
 
 	/** Getter function to override, returning pass layer proxies to be passed to the render thread. (Proxy objects should be allocated from the provided allocator.) */
-	UE_API virtual bool GetProxy(FTraversalContext& InContext, FSceneRenderingBulkObjectAllocator& InFrameAllocator, FCompositeCorePassProxy*& OutProxy) const override;
+	UE_API virtual bool GetProxy(FCompositeTraversalContext& InContext, FSceneRenderingBulkObjectAllocator& InFrameAllocator, FCompositeCorePassProxy*& OutProxy) const override;
 
 	/** Get reflection & shadow caster actors. */
 	UFUNCTION(BlueprintGetter)
diff --git a/Engine/Plugins/Compositing/Composite/Source/Composite/Public/Layers/CompositeLayerShadowReflection.h b/Engine/Plugins/Compositing/Composite/Source/Composite/Public/Layers/CompositeLayerShadowReflection.h
--- a/Engine/Plugins/Compositing/Composite/Source/Composite/Public/Layers/CompositeLayerShadowReflection.h
+++ b/Engine/Plugins/Compositing/Composite/Source/Composite/Public/Layers/CompositeLayerShadowReflection.h
@@ -59,7 +59,7 @@ public:
 	UE_API virtual void OnEndOfFrameUpdate(UWorld* InWorld) override;
 
 	/** Getter function to override, returning pass layer proxies to be passed to the render thread. (Proxy objects should be allocated from the provided allocator.) */
-	UE_API virtual bool GetProxy(FTraversalContext& InContext, FSceneRenderingBulkObjectAllocator& InFrameAllocator, FCompositeCorePassProxy*& OutProxy) const override;
+	UE_API virtual bool GetProxy(FCompositeTraversalContext& InContext, FSceneRenderingBulkObjectAllocator& InFrameAllocator, FCompositeCorePassProxy*& OutProxy) const override;
 
 	/** Get reflection & shadow caster actors. */
 	UFUNCTION(BlueprintGetter)
diff --git a/Engine/Plugins/Compositing/Composite/Source/Composite/Public/Layers/CompositeLayerSingleLightShadow.h b/Engine/Plugins/Compositing/Composite/Source/Composite/Public/Layers/CompositeLayerSingleLightShadow.h
--- a/Engine/Plugins/Compositing/Composite/Source/Composite/Public/Layers/CompositeLayerSingleLightShadow.h
+++ b/Engine/Plugins/Compositing/Composite/Source/Composite/Public/Layers/CompositeLayerSingleLightShadow.h
@@ -52,7 +52,7 @@ public:
 	UE_API virtual void OnEndOfFrameUpdate(UWorld* InWorld) override;
 
 	/** Getter function to override, returning pass layer proxies to be passed to the render thread. (Proxy objects should be allocated from the provided allocator.) */
-	UE_API virtual bool GetProxy(FTraversalContext& InContext, FSceneRenderingBulkObjectAllocator& InFrameAllocator, FCompositeCorePassProxy*& OutProxy) const override;
+	UE_API virtual bool GetProxy(FCompositeTraversalContext& InContext, FSceneRenderingBulkObjectAllocator& InFrameAllocator, FCompositeCorePassProxy*& OutProxy) const override;
 
 public:
 	/**
diff --git a/Engine/Plugins/Compositing/Composite/Source/Composite/Public/Passes/CompositePassBase.h b/Engine/Plugins/Compositing/Composite/Source/Composite/Public/Passes/CompositePassBase.h
--- a/Engine/Plugins/Compositing/Composite/Source/Composite/Public/Passes/CompositePassBase.h
+++ b/Engine/Plugins/Compositing/Composite/Source/Composite/Public/Passes/CompositePassBase.h
@@ -12,6 +12,39 @@
 
 #define UE_API COMPOSITE_API
 
+/** Context used while traversing passes and layers to collect frame proxies. */
+struct FCompositeTraversalContext
+{
+	/** Flag to indicate solo pass */
+	bool bIsSolo = false;
+
+	/** Flag to indicate first pass */
+	bool bIsFirstPass = true;
+
+	/**
+	 * Flag to indicate that a layer is dependent on scene textures.
+	 * Work from scene-independent layers can be shared between post-processing locations.
+	 */
+	bool bNeedsSceneTextures = false;
+
+	/** Array of passes applied at the start of rendering. */
+	TSortedMap<UE::CompositeCore::ResourceId, TArray<const FCompositeCorePassProxy*>> PreprocessingPasses;
+
+	/**
+	* Find or create an external texture resource.
+	*
+	* @return ResourceId Texture resource identifier
+	*/
+	UE::CompositeCore::ResourceId FindOrCreateExternalTexture(TWeakObjectPtr<UTexture> InTexture, UE::CompositeCore::FResourceMetadata InMetadata);
+
+	/** Get external textures. */
+	const TArray<UE::CompositeCore::FExternalTexture>& GetExternalTextures() const;
+
+private:
+	/** List of external textures */
+	TArray<UE::CompositeCore::FExternalTexture> ExternalTextures;
+};
+
 /** Base class for defining a pass object, and its associated render proxy. */
 UCLASS(MinimalAPI, Abstract, EditInlineNew)
 class UCompositePassBase : public UObject
@@ -30,7 +63,7 @@ public:
 #endif
 
 	/** Getter function to override, returning pass proxies to be passed to the render thread. (Proxy objects should be allocated from the provided allocator.) */
-	UE_API virtual bool GetProxy(const UE::CompositeCore::FPassInputDecl& InputDecl, FSceneRenderingBulkObjectAllocator& InFrameAllocator, FCompositeCorePassProxy*& OutProxy) const { return false; };
+	UE_API virtual bool GetProxy(const UE::CompositeCore::FPassInputDecl& InputDecl, FCompositeTraversalContext& InContext, FSceneRenderingBulkObjectAllocator& InFrameAllocator, FCompositeCorePassProxy*& OutProxy) const { return false; };
 
 	/** Get the enabled state. */
 	UFUNCTION(BlueprintGetter)
@@ -43,9 +76,6 @@ public:
 	/** Returns true if the pass is returns a valid proxy. */
 	virtual bool IsActive() const { return IsEnabled(); }
 
-	/** Returns true if the pass (may) need post-processing scene textures. */
-	virtual bool NeedsSceneTextures() const { return false; }
-
 #if WITH_EDITOR
 	/** Gets the display name of the pass */
 	FString GetDisplayName() const { return DisplayName; }
diff --git a/Engine/Plugins/Compositing/Composite/Source/Composite/Public/Passes/CompositePassCenteredScale.h b/Engine/Plugins/Compositing/Composite/Source/Composite/Public/Passes/CompositePassCenteredScale.h
--- a/Engine/Plugins/Compositing/Composite/Source/Composite/Public/Passes/CompositePassCenteredScale.h
+++ b/Engine/Plugins/Compositing/Composite/Source/Composite/Public/Passes/CompositePassCenteredScale.h
@@ -9,7 +9,7 @@
 #define UE_API COMPOSITE_API
 
 /** Primary centered scale calculation mode.*/
-UENUM()
+UENUM(BlueprintType)
 enum class ECompositePassScaleMode : uint8
 {
 	/** Disabled */
@@ -26,7 +26,7 @@ enum class ECompositePassScaleMode : uint8
 };
 
 /** Overscan uncrop derivation mode. */
-UENUM()
+UENUM(BlueprintType)
 enum class ECompositePassOverscanUncropMode : uint8
 {
 	/** Disabled */
@@ -52,7 +52,7 @@ public:
 	/** Destructor */
 	UE_API ~UCompositePassCenteredScale();
 
-	UE_API virtual bool GetProxy(const UE::CompositeCore::FPassInputDecl& InputDecl, FSceneRenderingBulkObjectAllocator& InFrameAllocator, FCompositeCorePassProxy*& OutProxy) const override;
+	UE_API virtual bool GetProxy(const UE::CompositeCore::FPassInputDecl& InputDecl, FCompositeTraversalContext& InContext, FSceneRenderingBulkObjectAllocator& InFrameAllocator, FCompositeCorePassProxy*& OutProxy) const override;
 
 	UE_API virtual bool IsActive() const override;
 
diff --git a/Engine/Plugins/Compositing/Composite/Source/Composite/Public/Passes/CompositePassColorGrading.h b/Engine/Plugins/Compositing/Composite/Source/Composite/Public/Passes/CompositePassColorGrading.h
--- a/Engine/Plugins/Compositing/Composite/Source/Composite/Public/Passes/CompositePassColorGrading.h
+++ b/Engine/Plugins/Compositing/Composite/Source/Composite/Public/Passes/CompositePassColorGrading.h
@@ -49,7 +49,7 @@ public:
 	/** Destructor */
 	UE_API ~UCompositePassColorGrading();
 
-	UE_API virtual bool GetProxy(const UE::CompositeCore::FPassInputDecl& InputDecl, FSceneRenderingBulkObjectAllocator& InFrameAllocator, FCompositeCorePassProxy*& OutProxy) const override;
+	UE_API virtual bool GetProxy(const UE::CompositeCore::FPassInputDecl& InputDecl, FCompositeTraversalContext& InContext, FSceneRenderingBulkObjectAllocator& InFrameAllocator, FCompositeCorePassProxy*& OutProxy) const override;
 
 public:
 	/** Color is divided by its alpha channel before processing, and multiplied again after. */
diff --git a/Engine/Plugins/Compositing/Composite/Source/Composite/Public/Passes/CompositePassColorKeyer.h b/Engine/Plugins/Compositing/Composite/Source/Composite/Public/Passes/CompositePassColorKeyer.h
--- a/Engine/Plugins/Compositing/Composite/Source/Composite/Public/Passes/CompositePassColorKeyer.h
+++ b/Engine/Plugins/Compositing/Composite/Source/Composite/Public/Passes/CompositePassColorKeyer.h
@@ -50,7 +50,7 @@ public:
 	/** Destructor */
 	UE_API ~UCompositePassColorKeyer();
 
-	UE_API virtual bool GetProxy(const UE::CompositeCore::FPassInputDecl& InputDecl, FSceneRenderingBulkObjectAllocator& InFrameAllocator, FCompositeCorePassProxy*& OutProxy) const override;
+	UE_API virtual bool GetProxy(const UE::CompositeCore::FPassInputDecl& InputDecl, FCompositeTraversalContext& InContext, FSceneRenderingBulkObjectAllocator& InFrameAllocator, FCompositeCorePassProxy*& OutProxy) const override;
 
 #if WITH_EDITOR
 	UE_API virtual void PreEditChange(FProperty* PropertyThatWillChange) override;
diff --git a/Engine/Plugins/Compositing/Composite/Source/Composite/Public/Passes/CompositePassDistortion.h b/Engine/Plugins/Compositing/Composite/Source/Composite/Public/Passes/CompositePassDistortion.h
--- a/Engine/Plugins/Compositing/Composite/Source/Composite/Public/Passes/CompositePassDistortion.h
+++ b/Engine/Plugins/Compositing/Composite/Source/Composite/Public/Passes/CompositePassDistortion.h
@@ -37,7 +37,7 @@ public:
 
 	virtual bool IsActive() const override;
 
-	UE_API virtual bool GetProxy(const UE::CompositeCore::FPassInputDecl& InputDecl, FSceneRenderingBulkObjectAllocator& InFrameAllocator, FCompositeCorePassProxy*& OutProxy) const override;
+	UE_API virtual bool GetProxy(const UE::CompositeCore::FPassInputDecl& InputDecl, FCompositeTraversalContext& InContext, FSceneRenderingBulkObjectAllocator& InFrameAllocator, FCompositeCorePassProxy*& OutProxy) const override;
 
 public:
 	/** Distortion to apply on input. */
diff --git a/Engine/Plugins/Compositing/Composite/Source/Composite/Public/Passes/CompositePassFXAA.h b/Engine/Plugins/Compositing/Composite/Source/Composite/Public/Passes/CompositePassFXAA.h
--- a/Engine/Plugins/Compositing/Composite/Source/Composite/Public/Passes/CompositePassFXAA.h
+++ b/Engine/Plugins/Compositing/Composite/Source/Composite/Public/Passes/CompositePassFXAA.h
@@ -21,7 +21,7 @@ public:
 	/** Destructor */
 	UE_API ~UCompositePassFXAA();
 
-	UE_API virtual bool GetProxy(const UE::CompositeCore::FPassInputDecl& InputDecl, FSceneRenderingBulkObjectAllocator& InFrameAllocator, FCompositeCorePassProxy*& OutProxy) const override;
+	UE_API virtual bool GetProxy(const UE::CompositeCore::FPassInputDecl& InputDecl, FCompositeTraversalContext& InContext, FSceneRenderingBulkObjectAllocator& InFrameAllocator, FCompositeCorePassProxy*& OutProxy) const override;
 
 public:
 
diff --git a/Engine/Plugins/Compositing/Composite/Source/Composite/Public/Passes/CompositePassMasking.h b/Engine/Plugins/Compositing/Composite/Source/Composite/Public/Passes/CompositePassMasking.h
new file mode 100644
--- /dev/null
+++ b/Engine/Plugins/Compositing/Composite/Source/Composite/Public/Passes/CompositePassMasking.h
@@ -0,0 +1,57 @@
+// Copyright Epic Games, Inc. All Rights Reserved.
+
+#pragma once
+
+#include "CompositePassBase.h"
+
+#include "CompositePassMasking.generated.h"
+
+#define UE_API COMPOSITE_API
+
+/** Color channel. */
+UENUM(BlueprintType)
+enum class ECompositeColorChannel : uint8
+{
+	Red,
+	Green,
+	Blue,
+	Alpha
+};
+
+/** Alpha masking pass. */
+UCLASS(MinimalAPI, BlueprintType, Blueprintable, EditInlineNew, CollapseCategories, meta=(DisplayName="Masking Pass"))
+class UCompositePassMasking : public UCompositePassBase
+{
+	GENERATED_BODY()
+
+public:
+	/** Constructor */
+	UE_API UCompositePassMasking(const FObjectInitializer& ObjectInitializer);
+	
+	/** Destructor */
+	UE_API ~UCompositePassMasking();
+
+	UE_API virtual bool IsActive() const override;
+
+	UE_API virtual bool GetProxy(const UE::CompositeCore::FPassInputDecl& InputDecl, FCompositeTraversalContext& InContext, FSceneRenderingBulkObjectAllocator& InFrameAllocator, FCompositeCorePassProxy*& OutProxy) const override;
+
+public:
+	/** Masking texture. */
+	UPROPERTY(EditAnywhere, BlueprintReadWrite, Category = "", meta = (AllowedClasses = "/Script/Engine.Texture2D, /Script/MediaAssets.MediaTexture"))
+	TObjectPtr<UTexture> MaskTexture;
+
+	/** Color channel used as the mask source. */
+	UPROPERTY(BlueprintReadWrite, EditAnywhere, Category = "")
+	ECompositeColorChannel MaskSource;
+
+	/** Color is divided by its alpha channel before masking, and multiplied again after. */
+	UPROPERTY(BlueprintReadWrite, EditAnywhere, Category = "")
+	bool bUnpremult;
+
+	/** Invert mask alpha. */
+	UPROPERTY(BlueprintReadWrite, EditAnywhere, Category = "")
+	bool bInvert;
+};
+
+#undef UE_API
+
diff --git a/Engine/Plugins/Compositing/Composite/Source/Composite/Public/Passes/CompositePassMaterial.h b/Engine/Plugins/Compositing/Composite/Source/Composite/Public/Passes/CompositePassMaterial.h
--- a/Engine/Plugins/Compositing/Composite/Source/Composite/Public/Passes/CompositePassMaterial.h
+++ b/Engine/Plugins/Compositing/Composite/Source/Composite/Public/Passes/CompositePassMaterial.h
@@ -25,9 +25,7 @@ public:
 
 	UE_API virtual bool IsActive() const override;
 
-	UE_API virtual bool GetProxy(const UE::CompositeCore::FPassInputDecl& InputDecl, FSceneRenderingBulkObjectAllocator& InFrameAllocator, FCompositeCorePassProxy*& OutProxy) const override;
-
-	UE_API virtual bool NeedsSceneTextures() const override;
+	UE_API virtual bool GetProxy(const UE::CompositeCore::FPassInputDecl& InputDecl, FCompositeTraversalContext& InContext, FSceneRenderingBulkObjectAllocator& InFrameAllocator, FCompositeCorePassProxy*& OutProxy) const override;
 
 public:
 	/** Post-process material to execute. Input0 is connected to SceneTexture's PostprocessInput0. */
diff --git a/Engine/Plugins/Compositing/Composite/Source/Composite/Public/Passes/CompositePassOpenColorIO.h b/Engine/Plugins/Compositing/Composite/Source/Composite/Public/Passes/CompositePassOpenColorIO.h
--- a/Engine/Plugins/Compositing/Composite/Source/Composite/Public/Passes/CompositePassOpenColorIO.h
+++ b/Engine/Plugins/Compositing/Composite/Source/Composite/Public/Passes/CompositePassOpenColorIO.h
@@ -24,7 +24,7 @@ public:
 	/** Destructor */
 	UE_API ~UCompositePassOpenColorIO();
 
-	UE_API virtual bool GetProxy(const UE::CompositeCore::FPassInputDecl& InputDecl, FSceneRenderingBulkObjectAllocator& InFrameAllocator, FCompositeCorePassProxy*& OutProxy) const override;
+	UE_API virtual bool GetProxy(const UE::CompositeCore::FPassInputDecl& InputDecl, FCompositeTraversalContext& InContext, FSceneRenderingBulkObjectAllocator& InFrameAllocator, FCompositeCorePassProxy*& OutProxy) const override;
 
 public:
 	/** Color conversion settings. */
diff --git a/Engine/Plugins/Compositing/Composite/Source/Composite/Public/Passes/CompositePassSMAA.h b/Engine/Plugins/Compositing/Composite/Source/Composite/Public/Passes/CompositePassSMAA.h
--- a/Engine/Plugins/Compositing/Composite/Source/Composite/Public/Passes/CompositePassSMAA.h
+++ b/Engine/Plugins/Compositing/Composite/Source/Composite/Public/Passes/CompositePassSMAA.h
@@ -21,11 +21,11 @@ public:
 	/** Destructor */
 	UE_API ~UCompositePassSMAA();
 
-	UE_API virtual bool GetProxy(const UE::CompositeCore::FPassInputDecl& InputDecl, FSceneRenderingBulkObjectAllocator& InFrameAllocator, FCompositeCorePassProxy*& OutProxy) const override;
+	UE_API virtual bool GetProxy(const UE::CompositeCore::FPassInputDecl& InputDecl, FCompositeTraversalContext& InContext, FSceneRenderingBulkObjectAllocator& InFrameAllocator, FCompositeCorePassProxy*& OutProxy) const override;
 
 public:
 
-	/** FXAA quality control, matching the "r.fxaa.quality" console variable values. */
+	/** SMAA quality control, matching the "r.SMAA.quality" console variable values. */
 	UPROPERTY(EditAnywhere, BlueprintReadWrite, Category = "", meta=(ClampMin="0", ClampMax="5", UIMin = "0", UIMax = "5"))
 	int32 Quality;
 };
diff --git a/Engine/Plugins/Compositing/CompositeCore/Source/CompositeCore/Private/CompositeCoreSceneViewExtension.cpp b/Engine/Plugins/Compositing/CompositeCore/Source/CompositeCore/Private/CompositeCoreSceneViewExtension.cpp
--- a/Engine/Plugins/Compositing/CompositeCore/Source/CompositeCore/Private/CompositeCoreSceneViewExtension.cpp
+++ b/Engine/Plugins/Compositing/CompositeCore/Source/CompositeCore/Private/CompositeCoreSceneViewExtension.cpp
@@ -479,46 +479,20 @@ void FCompositeCoreSceneViewExtension::ApplyPreprocessing(FRDGBuilder& GraphBuil
 
 		// Setup pass inputs
 		FPassInputArray PassInputs;
-		PassInputs.GetArray().Add(*ExternalTexture);
-
-		// Iterate over all passes
-		for (int32 PassIndex =0; PassIndex <Passes.Num(); ++PassIndex)
+		PassInputs.GetArray().SetNum(kPostProcessMaterialInputCountMax);
+		PassInputs[0] = *ExternalTexture;
+		// Empty textures for the other inputs
+		for (int32 Index = 1; Index < kPostProcessMaterialInputCountMax; ++Index)
 		{
-			const FCompositeCorePassProxy* Pass = Passes[PassIndex];
-			const TOptional<UE::CompositeCore::ResourceId>& OverrrideOutput = Pass->GetDeclaredOutputOverride();
-			bool bPassOutputOverridden = false;
-
-			if (OverrrideOutput.IsSet())
-			{
-				const FPassInput* ExternalOutputOverride = ExternalTextures.Get().Find(OverrrideOutput.GetValue());
-				if (ensureMsgf(ExternalOutputOverride != nullptr, TEXT("Unexpected missing external render target override as output.")))
-				{
-					PassInputs.OverrideOutput = FScreenPassRenderTarget(ExternalOutputOverride->Texture, ERenderTargetLoadAction::ENoAction);
-					
-					// Update view rect to match external texture
-					PassContext.OutputViewRect = ExternalOutputOverride->Texture.ViewRect;
-					// Mark as overriden
-					bPassOutputOverridden = true;
-				}
-			}
-			
-
-			if(!bPassOutputOverridden)
-			{
-				PassInputs.OverrideOutput = FScreenPassRenderTarget{};
-
-				// Update view rect to match external texture
-				PassContext.OutputViewRect = ExternalTexture->Texture.ViewRect;
-			}
-
-			// Register pass and update output
-			const FPassTexture Output = Pass->Add(GraphBuilder, InView, PassInputs, PassContext);
+			PassInputs[Index].Texture = FScreenPassTexture{ GSystemTextures.GetBlackDummy(GraphBuilder) };
+		}
 
-			// Note: we do not cache the output to CachedOutputsPerView_RenderThread since we only support cached outputs in post-processing for now.
+		FPassTexture Output;
+		PassContext.OutputViewRect = ExternalTexture->Texture.ViewRect;
+		constexpr int32 LastMergePassIndex = -1;
+		constexpr int32 RecursionLevel = 0;
 
-			// Update next pass input
-			PassInputs[0] = Output;
-		}
+		ApplyPassesRecursive(GraphBuilder, InView, PassInputs, PassInputs, PassContext, Passes, ExternalTextures, LastMergePassIndex, RecursionLevel, Output);
 	}
 }
 
@@ -591,7 +565,7 @@ UE::CompositeCore::FScopedExternalTextureMap::~FScopedExternalTextureMap()
 	}
 }
 
-bool FCompositeCoreSceneViewExtension::ApplyPasses_Recursive(
+bool FCompositeCoreSceneViewExtension::ApplyPassesRecursive(
 	FRDGBuilder& GraphBuilder,
 	const FSceneView& InView,
 	const UE::CompositeCore::FPassInputArray& Inputs,
@@ -601,7 +575,7 @@ bool FCompositeCoreSceneViewExtension::ApplyPasses_Recursive(
 	const UE::CompositeCore::FScopedExternalTextureMap& InExternalTextures,
 	const int32 LastMergePassIndex,
 	const int32 RecursionLevel,
-	UE::CompositeCore::FPassTexture& Output)
+	UE::CompositeCore::FPassTexture& Output) const
 {
 	using namespace UE::CompositeCore;
 
@@ -649,7 +623,7 @@ bool FCompositeCoreSceneViewExtension::ApplyPasses_Recursive(
 				{
 					PassInputs[InputIndex] = *PreProducedOutput;
 				}
-				else if(ApplyPasses_Recursive(GraphBuilder, InView, PassInputs, OriginalInputs, PassContext, { ChildPassProxy }, InExternalTextures, INDEX_NONE, RecursionLevel + 1, ChildPassOutput))
+				else if(ApplyPassesRecursive(GraphBuilder, InView, PassInputs, OriginalInputs, PassContext, { ChildPassProxy }, InExternalTextures, INDEX_NONE, RecursionLevel + 1, ChildPassOutput))
 				{
 					PassInputs[InputIndex] = ChildPassOutput;
 				}
@@ -700,16 +674,32 @@ bool FCompositeCoreSceneViewExtension::ApplyPasses_Recursive(
 			}
 		}
 
-		if (bIsLastMergePass)
+		const TOptional<UE::CompositeCore::ResourceId>& OverrrideOutput = Pass->GetDeclaredOutputOverride();
+
+		if (OverrrideOutput.IsSet())
 		{
-			// Merge pass proxies will handle alpha inversion, lens distortion and color encodings when writing back to scene color
-			PassContext.bOutputSceneColor = true;
-		}
+			const FPassInput* ExternalOutputOverride = InExternalTextures.Get().Find(OverrrideOutput.GetValue());
+			if (ensureMsgf(ExternalOutputOverride != nullptr, TEXT("Unexpected missing external render target override as output.")))
+			{
+				PassInputs.OverrideOutput = FScreenPassRenderTarget(ExternalOutputOverride->Texture, ERenderTargetLoadAction::ENoAction);
 
-		if (!bIsLastPass)
+				// Update view rect to match external texture
+				PassContext.OutputViewRect = ExternalOutputOverride->Texture.ViewRect;
+			}
+		}
+		else
 		{
-			// Only apply the output override on the last pass
-			PassInputs.OverrideOutput = FScreenPassRenderTarget{};
+			if (bIsLastMergePass)
+			{
+				// Merge pass proxies will handle alpha inversion, lens distortion and color encodings when writing back to scene color
+				PassContext.bOutputSceneColor = true;
+			}
+
+			if (!bIsLastPass)
+			{
+				// Only apply the output override on the last pass
+				PassInputs.OverrideOutput = FScreenPassRenderTarget{};
+			}
 		}
 
 		// Register pass and update output
@@ -746,7 +736,7 @@ FScreenPassTexture FCompositeCoreSceneViewExtension::PostProcessWork_RenderThrea
 		});
 	constexpr int32 RecursionLevel = 0;
 	
-	if (ApplyPasses_Recursive(GraphBuilder, InView, ResolvedInputs, ResolvedInputs, PassContext, Passes, ExternalTextures, LastMergePassIndex, RecursionLevel, Output))
+	if (ApplyPassesRecursive(GraphBuilder, InView, ResolvedInputs, ResolvedInputs, PassContext, Passes, ExternalTextures, LastMergePassIndex, RecursionLevel, Output))
 	{
 		return MoveTemp(Output.Texture);
 	}
diff --git a/Engine/Plugins/Compositing/CompositeCore/Source/CompositeCore/Private/CompositeCoreSceneViewExtension.h b/Engine/Plugins/Compositing/CompositeCore/Source/CompositeCore/Private/CompositeCoreSceneViewExtension.h
--- a/Engine/Plugins/Compositing/CompositeCore/Source/CompositeCore/Private/CompositeCoreSceneViewExtension.h
+++ b/Engine/Plugins/Compositing/CompositeCore/Source/CompositeCore/Private/CompositeCoreSceneViewExtension.h
@@ -116,7 +116,7 @@ private:
 	void ApplyPreprocessing(FRDGBuilder& GraphBuilder, FSceneView& InView) const;
 
 	/** Recursive pass application. */
-	bool ApplyPasses_Recursive(
+	bool ApplyPassesRecursive(
 		FRDGBuilder& GraphBuilder,
 		const FSceneView& InView,
 		const UE::CompositeCore::FPassInputArray& Inputs,
@@ -127,7 +127,7 @@ private:
 		const int32 LastMergePassIndex,
 		const int32 RecursionLevel,
 		UE::CompositeCore::FPassTexture& Output
-	);
+	) const;
 
 	/** Callback for processing passes. */
 	FScreenPassTexture PostProcessWork_RenderThread(FRDGBuilder& GraphBuilder, const FSceneView& InView,
diff --git a/Engine/Plugins/Compositing/CompositeCore/Source/CompositeCore/Public/Passes/CompositeCorePassProxy.h b/Engine/Plugins/Compositing/CompositeCore/Source/CompositeCore/Public/Passes/CompositeCorePassProxy.h
--- a/Engine/Plugins/Compositing/CompositeCore/Source/CompositeCore/Public/Passes/CompositeCorePassProxy.h
+++ b/Engine/Plugins/Compositing/CompositeCore/Source/CompositeCore/Public/Passes/CompositeCorePassProxy.h
@@ -187,13 +187,13 @@ namespace UE
 		struct FPassContext
 		{
 			/** The uniform buffer containing all scene textures. */
-			FSceneTextureShaderParameters SceneTextures;
+			FSceneTextureShaderParameters SceneTextures = {};
 
 			/** Active post-processing output view rectangle. */
-			FIntRect OutputViewRect;
+			FIntRect OutputViewRect = {};
 			
 			/** Post-processiong location. */
-			ISceneViewExtension::EPostProcessingPass Location;
+			ISceneViewExtension::EPostProcessingPass Location = ISceneViewExtension::EPostProcessingPass::BeforeDOF;
 			
 			/** Is the current pass expected to output scene color? */
 			bool bOutputSceneColor = false;
