Composure: Display color space blending for Ultimatte Key & Fill.

#Synced-CL 52383054

diff --git a/Engine/Plugins/Compositing/Composite/Source/Composite/Private/Layers/CompositeLayerPlate.cpp b/Engine/Plugins/Compositing/Composite/Source/Composite/Private/Layers/CompositeLayerPlate.cpp
--- a/Engine/Plugins/Compositing/Composite/Source/Composite/Private/Layers/CompositeLayerPlate.cpp
+++ b/Engine/Plugins/Compositing/Composite/Source/Composite/Private/Layers/CompositeLayerPlate.cpp
@@ -656,7 +656,7 @@ FCompositeCorePassProxy* UCompositeLayerPlate::GetProxy(const UE::CompositeCore:
 	Inputs[0] = PassInput;
 	Inputs[1] = GetDefaultSecondInput(InContext);
 
-	return InFrameAllocator.Create<FMergePassProxy>(MoveTemp(Inputs), GetMergeOperation(InContext), TEXT("Plate"));
+	return InFrameAllocator.Create<FMergePassProxy>(MoveTemp(Inputs), GetMergeOperation(InContext), TEXT("Plate"), ELensDistortionHandling::Enabled, bBlendInDisplaySpace);
 }
 
 TArray<UPrimitiveComponent*> UCompositeLayerPlate::GetPrimitives() const
diff --git a/Engine/Plugins/Compositing/Composite/Source/Composite/Public/Layers/CompositeLayerPlate.h b/Engine/Plugins/Compositing/Composite/Source/Composite/Public/Layers/CompositeLayerPlate.h
--- a/Engine/Plugins/Compositing/Composite/Source/Composite/Public/Layers/CompositeLayerPlate.h
+++ b/Engine/Plugins/Compositing/Composite/Source/Composite/Public/Layers/CompositeLayerPlate.h
@@ -181,6 +181,20 @@ private:
 	UPROPERTY(EditAnywhere, BlueprintGetter = GetControlMeshVisibility, BlueprintSetter = SetControlMeshVisibility, Category = "Composite", AdvancedDisplay, meta = (AllowPrivateAccess = true))
 	bool bControlMeshVisibility;
 
+public:
+	/**
+	 * When enabled, blending with the background is performed in sRGB display color space rather than linear.
+	 * Use for Ultimatte Key & Fill workflows where the content is in display space.
+	 * 
+	 * Addition Ultimatte requirements:
+	 * - "Encoding Override" should be set to "Linear" on the (alpha) key media source to avoid color decoding.
+	 * - The masking pass receiving the (alpha) key will have:
+	 *      - Unpremult: false
+	 *      - Invert: true
+	 */
+	UPROPERTY(EditAnywhere, BlueprintReadWrite, Category = "Composite", AdvancedDisplay)
+	bool bBlendInDisplaySpace = false;
+
 private:
 	/** Condtionally register composite mesh primitives: active holdout and add to the scene view extension custom render pass. */
 	void TryRegisterCompositePrimitives(const TArray<UPrimitiveComponent*>& InPrimitives, const UWorld* SourceWorld = nullptr) const;
diff --git a/Engine/Plugins/Compositing/CompositeCore/Shaders/Private/CompositeCoreMerge.usf b/Engine/Plugins/Compositing/CompositeCore/Shaders/Private/CompositeCoreMerge.usf
--- a/Engine/Plugins/Compositing/CompositeCore/Shaders/Private/CompositeCoreMerge.usf
+++ b/Engine/Plugins/Compositing/CompositeCore/Shaders/Private/CompositeCoreMerge.usf
@@ -65,6 +65,7 @@ float2 DisplayGamma; // DisplayGamma, 1/DisplayGamma
 uint MergeOp;
 uint OutputEncoding;
 uint bFlags; // 0: bInvertedAlphaOutput
+uint bBlendInDisplaySpace;
 
 float3 Decode(float3 InColor, uint InEncoding)
 {
@@ -165,6 +166,14 @@ void MainPS(
     Color0.rgb *= Tex0_Exposure;
     Color1.rgb *= Tex1_Exposure;
     
+    // Optionally encode both inputs to sRGB for display-space blending (e.g. Ultimatte Key & Fill)
+    BRANCH
+    if(bBlendInDisplaySpace)
+    {
+        Color0.rgb = LinearToSrgb(Color0.rgb);
+        Color1.rgb = LinearToSrgb(Color1.rgb);
+    }
+
     BRANCH
     if(MergeOp == ECompositeCoreMergeOp_Over)
     {
@@ -253,6 +262,13 @@ void MainPS(
         OutColor = Color0;
     }
 
+    // Decode back from sRGB to linear after display-space blending
+    BRANCH
+    if(bBlendInDisplaySpace)
+    {
+        OutColor.rgb = sRGBToLinear(OutColor.rgb);
+    }
+
     // Re-encode color
     OutColor.rgb = Encode(OutColor.rgb, OutputEncoding);
 
diff --git a/Engine/Plugins/Compositing/CompositeCore/Source/CompositeCore/Private/Passes/CompositeCorePassMergeProxy.cpp b/Engine/Plugins/Compositing/CompositeCore/Source/CompositeCore/Private/Passes/CompositeCorePassMergeProxy.cpp
--- a/Engine/Plugins/Compositing/CompositeCore/Source/CompositeCore/Private/Passes/CompositeCorePassMergeProxy.cpp
+++ b/Engine/Plugins/Compositing/CompositeCore/Source/CompositeCore/Private/Passes/CompositeCorePassMergeProxy.cpp
@@ -58,6 +58,7 @@ class FCompositeCoreMergeShader : public FGlobalShader
 		SHADER_PARAMETER(uint32, MergeOp)
 		SHADER_PARAMETER(uint32, OutputEncoding)
 		SHADER_PARAMETER(uint32, bFlags)
+		SHADER_PARAMETER(uint32, bBlendInDisplaySpace)
 		RENDER_TARGET_BINDING_SLOTS()
 	END_SHADER_PARAMETER_STRUCT()
 };
@@ -117,11 +118,12 @@ namespace UE
 			return Parameters;
 		}
 
-		FMergePassProxy::FMergePassProxy(FPassInputDeclArray InPassDeclaredInputs, ECompositeCoreMergeOp InMergeOp, const TCHAR* InParentName, ELensDistortionHandling bInUseLensDistortion)
+		FMergePassProxy::FMergePassProxy(FPassInputDeclArray InPassDeclaredInputs, ECompositeCoreMergeOp InMergeOp, const TCHAR* InParentName, ELensDistortionHandling bInUseLensDistortion, bool bInBlendInDisplaySpace)
 			: FCompositeCorePassProxy(MoveTemp(InPassDeclaredInputs))
 			, MergeOp(InMergeOp)
 			, ParentName(InParentName)
 			, bUseLensDistortion(bInUseLensDistortion)
+			, bBlendInDisplaySpace(bInBlendInDisplaySpace)
 		{
 		}
 
@@ -212,6 +214,7 @@ namespace UE
 			PassParameters->MergeOp = static_cast<uint32>(MergeOp);
 			PassParameters->OutputEncoding = static_cast<uint32>(OutMetadata.Encoding);
 			PassParameters->bFlags = static_cast<uint32>(PassContext.bOutputSceneColor);
+			PassParameters->bBlendInDisplaySpace = static_cast<uint32>(bBlendInDisplaySpace);
 			PassParameters->RenderTargets[0] = Output.GetRenderTargetBinding();
 
 			FCompositeCoreMergeShader::FPermutationDomain PermutationVector;
diff --git a/Engine/Plugins/Compositing/CompositeCore/Source/CompositeCore/Public/Passes/CompositeCorePassMergeProxy.h b/Engine/Plugins/Compositing/CompositeCore/Source/CompositeCore/Public/Passes/CompositeCorePassMergeProxy.h
--- a/Engine/Plugins/Compositing/CompositeCore/Source/CompositeCore/Public/Passes/CompositeCorePassMergeProxy.h
+++ b/Engine/Plugins/Compositing/CompositeCore/Source/CompositeCore/Public/Passes/CompositeCorePassMergeProxy.h
@@ -57,7 +57,7 @@ namespace UE
 			using FCompositeCorePassProxy::FCompositeCorePassProxy;
 
 			/** Constructor */
-			UE_API FMergePassProxy(FPassInputDeclArray InPassDeclaredInputs, ECompositeCoreMergeOp InMergeOp, const TCHAR* InParentName, ELensDistortionHandling bInUseLensDistortion = ELensDistortionHandling::Enabled);
+			UE_API FMergePassProxy(FPassInputDeclArray InPassDeclaredInputs, ECompositeCoreMergeOp InMergeOp, const TCHAR* InParentName, ELensDistortionHandling bInUseLensDistortion = ELensDistortionHandling::Enabled, bool bInBlendInDisplaySpace = false);
 
 			UE_API FPassTexture Add(FRDGBuilder& GraphBuilder, const FSceneView& InView, const FPassInputArray& Inputs, const FPassContext& PassContext) const override;
 
@@ -70,6 +70,9 @@ namespace UE
 
 			/** Flag to enable or disable application of engine-provided lens distortion LUTs, true by default. */
 			ELensDistortionHandling bUseLensDistortion = ELensDistortionHandling::Enabled;
+
+			/** When enabled, both inputs are encoded to sRGB before blending and decoded back to linear after. */
+			bool bBlendInDisplaySpace = false;
 		};
 
 		const FName& GetMergePassPassName();
