[Backout] - CL52557377
#rnx
Original CL Desc
-----------------------------------------------------------------
Composure: Expose erosion/dilation pass support, which was previously only used implicitely by the built-in custom render pass.

* Added erosion mode as a new shader permutation.
* Larger dilations/erosions are support on our new pass by running the same shader iteratively.

diff --git a/Engine/Plugins/Compositing/Composite/Source/Composite/Private/Passes/CompositePassDilation.cpp b/Engine/Plugins/Compositing/Composite/Source/Composite/Private/Passes/CompositePassDilation.cpp
deleted file mode 100644
diff --git a/Engine/Plugins/Compositing/Composite/Source/Composite/Private/Tests/CompositePassDilationTest.cpp b/Engine/Plugins/Compositing/Composite/Source/Composite/Private/Tests/CompositePassDilationTest.cpp
deleted file mode 100644
diff --git a/Engine/Plugins/Compositing/Composite/Source/Composite/Public/Passes/CompositePassDilation.h b/Engine/Plugins/Compositing/Composite/Source/Composite/Public/Passes/CompositePassDilation.h
deleted file mode 100644
diff --git a/Engine/Plugins/Compositing/CompositeCore/Shaders/Private/CompositeCoreDilate.usf b/Engine/Plugins/Compositing/CompositeCore/Shaders/Private/CompositeCoreDilate.usf
--- a/Engine/Plugins/Compositing/CompositeCore/Shaders/Private/CompositeCoreDilate.usf
+++ b/Engine/Plugins/Compositing/CompositeCore/Shaders/Private/CompositeCoreDilate.usf
@@ -1,19 +1,15 @@
 // Copyright Epic Games, Inc. All Rights Reserved.
 
 #include "/Engine/Private/Common.ush"
-#include "/Engine/Private/ScreenPass.ush"
 
 #ifndef TILE_SIZE
     #error TILE_SIZE must be defined
 #endif
 
-#ifndef MORPHOLOGY_ERODE
-    #define MORPHOLOGY_ERODE 0
-#endif
-
 #define KERNEL_COUNT ((1 + 2 * DILATION_SIZE) * (1 + 2 * DILATION_SIZE) - 1)
+#define NEARLY_TRANSLUCENT_ALPHA 0.9999999f // Alpha is inverted opacity
 
-// Number of texels around the group tile for scene color.
+// Number of texels arround the group tile for scene color.
 #define LDS_TILE_BORDER_SIZE DILATION_SIZE
 
 // Width in texels of the depth tile cached into LDS.
@@ -24,24 +20,10 @@
 
 #define THREADGROUP_TOTAL (TILE_SIZE * TILE_SIZE)
 
-Texture2D Input_Texture;
-SamplerState Input_Sampler;
-SCREEN_PASS_TEXTURE_VIEWPORT(Input)
-SCREEN_PASS_TEXTURE_VIEWPORT(Output)
+Texture2D InputTexture;
 RWTexture2D<float4> RWOutputTexture;
-uint bInvertedAlpha;
+uint2 Dimensions;
 uint bOpacifyOutput;
-// Opacity at or below this threshold is treated as fully transparent for morphological operations.
-// In standard opacity space (0=transparent, 1=opaque): opacity < AlphaThreshold means "treat as transparent".
-float AlphaThreshold;
-
-// Normalize or restore alpha between inverted UE convention (A=0 opaque, A=1 transparent)
-// and standard opacity (0 transparent, 1 opaque). Self-inverse: apply twice to round-trip.
-// Defined inline — CompositeCommon.ush lives in the Composite plugin and is not accessible here.
-float InvertAlpha(float InAlpha, uint bInvert)
-{
-    return select(bInvert > 0u, 1.0f - InAlpha, InAlpha);
-}
 
 #if DILATION_SIZE == 1
 static const int2 kOffsets[KERNEL_COUNT] = {
@@ -99,15 +81,16 @@ float4 SampleCachedColors(in uint2 GroupThreadId, int2 PixelOffset = int2(0,0))
 
 float SafeRcp(float x)
 {
-    return x > 0.0 ? rcp(x) : 0.0;
+	return x > 0.0 ? rcp(x) : 0.0;
 }
 
-// Unpremultiply and force full opacity. Operates in standard opacity space (a = 1 = opaque).
-// Only used by CompositeCore's holdout path (bOpacifyOutput=1).
 float4 Opacify(float4 InColor)
 {
-    InColor.xyz *= SafeRcp(saturate(InColor.a));
-    return float4(InColor.xyz, 1.0f);
+    float RegularAlpha = saturate(1.0 - InColor.a);
+    // Unpremultiply color
+    InColor.xyz *= SafeRcp(RegularAlpha);
+    // Return opaque color, maintaining UE's inverted alpha convention
+	return float4(InColor.xyz, 0.0);
 }
 
 [numthreads(TILE_SIZE, TILE_SIZE, 1)]
@@ -119,12 +102,8 @@ void MainCS(
 {
     // Init LDS
     {
-        // GroupTexelOffset must be signed: when GroupId.x or GroupId.y is 0 the halo extends one
-        // pixel outside the viewport (position -1), which uint arithmetic would wrap to 0xFFFFFFFF.
-        // The subsequent min-clamp would then map to ViewportSize-1 (the opposite edge) rather
-        // than 0, loading pixels from the wrong corner and creating phantom opaque halos that look
-        // like diagonal propagation artefacts at the viewport boundary.
-        int2 GroupTexelOffsetSigned = int2(GroupId) * int2(TILE_SIZE, TILE_SIZE) - LDS_TILE_BORDER_SIZE;
+        // Global texel coordinate of the active group (corner) offset by the LDS border
+        uint2 GroupTexelOffset = GroupId * uint2(TILE_SIZE, TILE_SIZE) - LDS_TILE_BORDER_SIZE;
 
         // To account for the dilation margin, each warp thread should load extra texels.
         // We proceed linearly through the tile (with borders) and each warp thread loads
@@ -136,17 +115,15 @@ void MainCS(
         UNROLL
         for (uint i = 0; i < LoadCount; i++)
         {
-            // Compute the signed viewport-relative texel position and clamp to [0, ViewportSize-1].
-            int2 TexelLocationSigned = GroupTexelOffsetSigned + int2(
-                int(LinearGroupThreadId % LDS_TILE_WIDTH),
-                int(LinearGroupThreadId / LDS_TILE_WIDTH));
-            TexelLocationSigned = clamp(TexelLocationSigned, int2(0, 0), int2(Input_ViewportSize) - 1);
+            uint2 TexelLocation =  GroupTexelOffset + uint2(
+                LinearGroupThreadId % LDS_TILE_WIDTH,
+                LinearGroupThreadId / LDS_TILE_WIDTH);
 
             if ((LinearGroupThreadId < LDS_ARRAY_SIZE) ||
                 (i != LoadCount - 1) ||
                 (LDS_ARRAY_SIZE % THREADGROUP_TOTAL) == 0)
             {
-                SharedColor[LinearGroupThreadId] = Input_Texture[uint2(TexelLocationSigned) + uint2(Input_ViewportMin)];
+                SharedColor[LinearGroupThreadId] = InputTexture[TexelLocation];
             }
             LinearGroupThreadId += THREADGROUP_TOTAL;
         }
@@ -155,61 +132,37 @@ void MainCS(
     GroupMemoryBarrierWithGroupSync();
 
     // Discard out-of-bound texels
-    if (any(DispatchThreadId >= uint2(Output_ViewportSize)))
+    if (any(DispatchThreadId >= Dimensions))
     {
         return;
     }
 
     float4 TexelColor = SampleCachedColors(GroupThreadId);
-    // Normalize to standard opacity (0=transparent, 1=opaque) for all downstream logic.
-    TexelColor.a = InvertAlpha(TexelColor.a, bInvertedAlpha);
-
+    
 #if DILATION_SIZE > 0
-#if MORPHOLOGY_ERODE
-    // Erosion: if texel is nearly opaque and any neighbor is nearly transparent, make texel transparent.
-    BRANCH
-    if (TexelColor.a > 1.0f - AlphaThreshold)
-    {
-        UNROLL
-        for (int i = 0; i < KERNEL_COUNT; i++)
-        {
-            float4 NeighborColor = SampleCachedColors(GroupThreadId, kOffsets[i]);
-            if (InvertAlpha(NeighborColor.a, bInvertedAlpha) < AlphaThreshold)
-            {
-                TexelColor = float4(0.0f, 0.0f, 0.0f, 0.0f); // transparent black (standard opacity)
-                break;
-            }
-        }
-    }
-#else
-    // Dilation: if texel is nearly transparent and any neighbor is more opaque, take neighbor color.
+	// Apply (gather) dilation on any translucent or nearly translucent texel
     BRANCH
-    if (TexelColor.a < AlphaThreshold)
+    if(TexelColor.a > NEARLY_TRANSLUCENT_ALPHA)
     {
         UNROLL
-        for (int i = 0; i < KERNEL_COUNT; i++)
+        for(int i=0; i<KERNEL_COUNT; i++)
         {
             float4 OffsetColor = SampleCachedColors(GroupThreadId, kOffsets[i]);
-            float OffsetOpacity = InvertAlpha(OffsetColor.a, bInvertedAlpha);
 
-            if (OffsetOpacity > TexelColor.a)
+            if(OffsetColor.a < TexelColor.a)
             {
-                TexelColor = float4(OffsetColor.rgb, OffsetOpacity);
+                TexelColor = OffsetColor;
                 break; // We break as soon as we find a more opaque color
             }
         }
     }
-#endif // MORPHOLOGY_ERODE
-#endif // DILATION_SIZE > 0
+#endif
 
-    // Conditionally & fully opacify non-translucent colors for the over blend.
-    // Only used by CompositeCore's holdout path (bOpacifyOutput=1).
-    if (bOpacifyOutput && TexelColor.a > 1.0f - AlphaThreshold)
-    {
-        TexelColor = Opacify(TexelColor);
-    }
+	// Condtionally & fully opacify non-translucent colors for the over blend.
+	if (bOpacifyOutput && TexelColor.a < NEARLY_TRANSLUCENT_ALPHA)
+	{
+		TexelColor = Opacify(TexelColor);
+	}
 
-    // Restore original alpha convention before writing.
-    TexelColor.a = InvertAlpha(TexelColor.a, bInvertedAlpha);
-    RWOutputTexture[Output_ViewportMin + DispatchThreadId] = TexelColor;
+	RWOutputTexture[DispatchThreadId] = TexelColor;
 }
diff --git a/Engine/Plugins/Compositing/CompositeCore/Source/CompositeCore/Private/Passes/CompositeCorePassDilate.cpp b/Engine/Plugins/Compositing/CompositeCore/Source/CompositeCore/Private/Passes/CompositeCorePassDilate.cpp
--- a/Engine/Plugins/Compositing/CompositeCore/Source/CompositeCore/Private/Passes/CompositeCorePassDilate.cpp
+++ b/Engine/Plugins/Compositing/CompositeCore/Source/CompositeCore/Private/Passes/CompositeCorePassDilate.cpp
@@ -6,7 +6,6 @@
 #include "RenderGraphBuilder.h"
 #include "RenderGraphUtils.h"
 #include "PixelShaderUtils.h"
-#include "ScreenPass.h"
 #include "ProfilingDebugging/RealtimeGPUProfiler.h"
 
 DECLARE_GPU_STAT_NAMED(FCompositeCoreDilate, TEXT("CompositeCore.Dilate"));
@@ -20,12 +19,10 @@ class FCompositeCoreDilateShader : public FGlobalShader
 	using FPermutationDomain = TShaderPermutationDomain<FDilationSize>;
 
 	BEGIN_SHADER_PARAMETER_STRUCT(FParameters, )
-		SHADER_PARAMETER_STRUCT(FScreenPassTextureInput, Input)
-		SHADER_PARAMETER_STRUCT(FScreenPassTextureViewportParameters, Output)
+		SHADER_PARAMETER_RDG_TEXTURE(Texture2D, InputTexture)
 		SHADER_PARAMETER_RDG_TEXTURE_UAV(RWTexture2D<float4>, RWOutputTexture)
-		SHADER_PARAMETER(uint32, bInvertedAlpha)
+		SHADER_PARAMETER(FIntPoint, Dimensions)
 		SHADER_PARAMETER(uint32, bOpacifyOutput)
-		SHADER_PARAMETER(float, AlphaThreshold)
 	END_SHADER_PARAMETER_STRUCT()
 
 	static inline void ModifyCompilationEnvironment(const FGlobalShaderPermutationParameters& Parameters, FShaderCompilerEnvironment& OutEnvironment)
@@ -49,24 +46,14 @@ namespace UE
 			{
 				RDG_EVENT_SCOPE_STAT(GraphBuilder, FCompositeCoreDilate, "CompositeCore.Dilate");
 
-				checkf(Input->Desc.Extent == Output->Desc.Extent,
-					TEXT("AddDilatePass: Input (%dx%d) and Output (%dx%d) must have matching extents"),
-					Input->Desc.Extent.X, Input->Desc.Extent.Y, Output->Desc.Extent.X, Output->Desc.Extent.Y);
-
 				FGlobalShaderMap* GlobalShaderMap = GetGlobalShaderMap(FeatureLevel);
 				const FIntPoint TextureSize = Input->Desc.Extent;
 
-				// CompositeCore's dilation texture always fills the entire extent with no viewport offset.
-				FScreenPassTexture InputScreenPass(Input, FIntRect(FIntPoint::ZeroValue, TextureSize));
-
 				FCompositeCoreDilateShader::FParameters* PassParameters = GraphBuilder.AllocParameters<FCompositeCoreDilateShader::FParameters>();
-				PassParameters->Input = GetScreenPassTextureInput(InputScreenPass, TStaticSamplerState<SF_Point>::GetRHI());
-				FScreenPassTexture OutputScreenPass(Output, FIntRect(FIntPoint::ZeroValue, TextureSize));
-				PassParameters->Output = GetScreenPassTextureViewportParameters(FScreenPassTextureViewport(OutputScreenPass));
+				PassParameters->InputTexture = Input;
 				PassParameters->RWOutputTexture = GraphBuilder.CreateUAV(Output);
-				PassParameters->bInvertedAlpha = 1; // CompositeCore always uses inverted alpha (A=0 opaque, A=1 transparent)
+				PassParameters->Dimensions = TextureSize;
 				PassParameters->bOpacifyOutput = PassInputs.bOpacifyOutput;
-				PassParameters->AlphaThreshold = PassInputs.AlphaThreshold;
 
 				FCompositeCoreDilateShader::FPermutationDomain PermutationVector;
 				PermutationVector.Set<FCompositeCoreDilateShader::FDilationSize>(FMath::Clamp(PassInputs.DilationSize, 0, 2));
@@ -101,3 +88,4 @@ namespace UE
 		}
 	}
 }
+
diff --git a/Engine/Plugins/Compositing/CompositeCore/Source/CompositeCore/Private/Passes/CompositeCorePassDilate.h b/Engine/Plugins/Compositing/CompositeCore/Source/CompositeCore/Private/Passes/CompositeCorePassDilate.h
--- a/Engine/Plugins/Compositing/CompositeCore/Source/CompositeCore/Private/Passes/CompositeCorePassDilate.h
+++ b/Engine/Plugins/Compositing/CompositeCore/Source/CompositeCore/Private/Passes/CompositeCorePassDilate.h
@@ -13,14 +13,11 @@ namespace UE
 		{
 			struct FDilateInputs
 			{
-				/** Dilation kernel radius in pixels (0 = no-op, 1 = 3×3 neighborhood, 2 = 5×5 neighborhood). Clamped to [0, 2]. */
+				/** Size of the dilation step. WARNING: Not currently implemented. */
 				int32 DilationSize = 1;
-
+				
 				/** Opacify the pass output to solid colors. */
 				bool bOpacifyOutput = true;
-
-				/** Opacity threshold - pixels at or below this (standard opacity space) are treated as fully transparent. */
-				float AlphaThreshold = 0.0000001f;
 			};
 
 			/**
