CompositeCore: Cleanup unused/dead shader parameter struct.

#Synced-CL 52419094

diff --git a/Engine/Plugins/Compositing/CompositeCore/Source/CompositeCore/Private/CompositeCoreSceneViewExtension.cpp b/Engine/Plugins/Compositing/CompositeCore/Source/CompositeCore/Private/CompositeCoreSceneViewExtension.cpp
--- a/Engine/Plugins/Compositing/CompositeCore/Source/CompositeCore/Private/CompositeCoreSceneViewExtension.cpp
+++ b/Engine/Plugins/Compositing/CompositeCore/Source/CompositeCore/Private/CompositeCoreSceneViewExtension.cpp
@@ -31,11 +31,6 @@ static TAutoConsoleVariable<int32> CVarCompositeCoreApplyFXAA(
 	TEXT("When enabled, the custom render pass automatically applies FXAA."),
 	ECVF_RenderThreadSafe);
 
-/** Parameters used to transition a texture to SRVGraphics. */
-BEGIN_SHADER_PARAMETER_STRUCT(FCompositeRenderTargetTransitionParameters, )
-	RDG_TEXTURE_ACCESS(TransitionTexture, ERHIAccess::SRVGraphics)
-END_SHADER_PARAMETER_STRUCT()
-
 class FCompositeCoreCustomRenderPass : public FCustomRenderPassBase
 {
 public:
