CompositeCore: Fix incorrect opacification caused by CL 52560406, since alpha was now normalized.

#Synced-CL 52419094

diff --git a/Engine/Plugins/Compositing/CompositeCore/Shaders/Private/CompositeCoreDilate.usf b/Engine/Plugins/Compositing/CompositeCore/Shaders/Private/CompositeCoreDilate.usf
--- a/Engine/Plugins/Compositing/CompositeCore/Shaders/Private/CompositeCoreDilate.usf
+++ b/Engine/Plugins/Compositing/CompositeCore/Shaders/Private/CompositeCoreDilate.usf
@@ -103,7 +103,6 @@ float SafeRcp(float x)
 }
 
 // Unpremultiply and force full opacity. Operates in standard opacity space (a = 1 = opaque).
-// Only used by CompositeCore's holdout path (bOpacifyOutput=1).
 float4 Opacify(float4 InColor)
 {
     InColor.xyz *= SafeRcp(saturate(InColor.a));
@@ -203,8 +202,7 @@ void MainCS(
 #endif // DILATION_SIZE > 0
 
     // Conditionally & fully opacify non-translucent colors for the over blend.
-    // Only used by CompositeCore's holdout path (bOpacifyOutput=1).
-    if (bOpacifyOutput && TexelColor.a > 1.0f - AlphaThreshold)
+    if (bOpacifyOutput && TexelColor.a > AlphaThreshold)
     {
         TexelColor = Opacify(TexelColor);
     }
