Promote Composite (aka new Composure) and CompositeCore plugins from Experimental to Beta.

#Synced-CL 52687094

diff --git a/Engine/Plugins/Compositing/Composite/Composite.uplugin b/Engine/Plugins/Compositing/Composite/Composite.uplugin
--- a/Engine/Plugins/Compositing/Composite/Composite.uplugin
+++ b/Engine/Plugins/Compositing/Composite/Composite.uplugin
@@ -1,7 +1,7 @@
 {
-	"FileVersion": 1,
+	"FileVersion": 3,
 	"Version": 1,
-	"VersionName": "1.0",
+	"VersionName": "1.1",
 	"FriendlyName": "Composure",
 	"Description": "Modern system for real-time compositing. This plugin succeeds legacy Composure and extends CompositeCore.",
 	"Category": "Compositing",
@@ -12,8 +12,8 @@
 	"SupportURL": "",
 	"EnabledByDefault": false,
 	"CanContainContent": true,
-	"IsBetaVersion": false,
-	"IsExperimentalVersion": true,
+	"IsBetaVersion": true,
+	"IsExperimentalVersion": false,
 	"Installed": false,
 	"Modules": [
 		{
diff --git a/Engine/Plugins/Compositing/CompositeCore/CompositeCore.uplugin b/Engine/Plugins/Compositing/CompositeCore/CompositeCore.uplugin
--- a/Engine/Plugins/Compositing/CompositeCore/CompositeCore.uplugin
+++ b/Engine/Plugins/Compositing/CompositeCore/CompositeCore.uplugin
@@ -1,7 +1,7 @@
 {
-	"FileVersion": 1,
+	"FileVersion": 3,
 	"Version": 1,
-	"VersionName": "1.0",
+	"VersionName": "1.1",
 	"FriendlyName": "CompositeCore",
 	"Description": "Extensible core plugin for real-time compositing, with a default (holdout) composite pipeline through post-processing.",
 	"Category": "Compositing",
@@ -12,8 +12,8 @@
 	"SupportURL": "",
 	"EnabledByDefault": false,
 	"CanContainContent": false,
-	"IsBetaVersion": false,
-	"IsExperimentalVersion": true,
+	"IsBetaVersion": true,
+	"IsExperimentalVersion": false,
 	"Installed": false,
 	"Modules": [
 		{
