CompositeCore: Remove FXAA plugin setting in preparation for making this functionality only available in Composite.

#Synced-CL 52740614

diff --git a/Engine/Plugins/Compositing/CompositeCore/Source/CompositeCore/Private/CompositeCoreSettings.cpp b/Engine/Plugins/Compositing/CompositeCore/Source/CompositeCore/Private/CompositeCoreSettings.cpp
--- a/Engine/Plugins/Compositing/CompositeCore/Source/CompositeCore/Private/CompositeCoreSettings.cpp
+++ b/Engine/Plugins/Compositing/CompositeCore/Source/CompositeCore/Private/CompositeCoreSettings.cpp
@@ -26,7 +26,6 @@ static FAutoConsoleVariableRef CVarRefreshPrimitivesOnTick(
 
 UCompositeCorePluginSettings::UCompositeCorePluginSettings()
 	: bApplyPreExposure(0)
-	, bApplyFXAA(0)
 	, DisabledPrimitiveClasses(
 		{
 			FSoftClassPath("/Script/Engine.CapsuleComponent"),
diff --git a/Engine/Plugins/Compositing/CompositeCore/Source/CompositeCore/Public/CompositeCoreSettings.h b/Engine/Plugins/Compositing/CompositeCore/Source/CompositeCore/Public/CompositeCoreSettings.h
--- a/Engine/Plugins/Compositing/CompositeCore/Source/CompositeCore/Public/CompositeCoreSettings.h
+++ b/Engine/Plugins/Compositing/CompositeCore/Source/CompositeCore/Public/CompositeCoreSettings.h
@@ -57,12 +57,6 @@ public:
 		ConfigRestartRequired = false))
 	uint32 bApplyPreExposure : 1;
 
-	UPROPERTY(config, EditAnywhere, Category = CompositeCore, meta = (
-		ConsoleVariable = "CompositeCore.ApplyFXAA", DisplayName = "Apply FXAA",
-		ToolTip = "When enabled, FXAA is applied onto the separate composited render. Quality is controlled with \"r.FXAA.Quality\". (Maps to \"CompositeCore.ApplyFXAA\" console variable).",
-		ConfigRestartRequired = false))
-	uint32 bApplyFXAA : 1;
-
 	/** Primitive component classes that do not support the composite pipeline.*/
 	UPROPERTY(config, EditAnywhere, Category = CompositeCore)
 	TArray<FSoftClassPath> DisabledPrimitiveClasses;
