Composure: Fix potential FXAA/SMAA crash with missing output render target, when running as the first media pre-processing pass.

#Synced-CL 52740614

diff --git a/Engine/Plugins/Compositing/Composite/Source/Composite/Private/Passes/CompositePassSMAA.cpp b/Engine/Plugins/Compositing/Composite/Source/Composite/Private/Passes/CompositePassSMAA.cpp
--- a/Engine/Plugins/Compositing/Composite/Source/Composite/Private/Passes/CompositePassSMAA.cpp
+++ b/Engine/Plugins/Compositing/Composite/Source/Composite/Private/Passes/CompositePassSMAA.cpp
@@ -53,7 +53,9 @@ namespace UE
 				SMAAInputs.Quality = static_cast<ESMAAQuality>(FMath::Clamp(Quality, 0, static_cast<int32>(ESMAAQuality::MAX) - 1));
 				if (!bApplyDisplayTransform)
 				{
-					SMAAInputs.OverrideOutput = Inputs.OverrideOutput;
+					SMAAInputs.OverrideOutput = Inputs.OverrideOutput.IsValid()
+						? Inputs.OverrideOutput
+						: CreateOutputRenderTarget(GraphBuilder, InView, PassContext.OutputViewRect, Input.Texture->Desc, TEXT("CompositePassSMAA"));
 				}
 
 				FScreenPassTexture Output = AddSMAAPasses(GraphBuilder, InView, SMAAInputs);
diff --git a/Engine/Plugins/Compositing/CompositeCore/Source/CompositeCore/Private/Passes/CompositeCorePassFXAAProxy.cpp b/Engine/Plugins/Compositing/CompositeCore/Source/CompositeCore/Private/Passes/CompositeCorePassFXAAProxy.cpp
--- a/Engine/Plugins/Compositing/CompositeCore/Source/CompositeCore/Private/Passes/CompositeCorePassFXAAProxy.cpp
+++ b/Engine/Plugins/Compositing/CompositeCore/Source/CompositeCore/Private/Passes/CompositeCorePassFXAAProxy.cpp
@@ -106,7 +106,9 @@ namespace UE
 			FXAAInputs.Quality = static_cast<EFXAAQuality>(FMath::Clamp(Quality, 0, static_cast<int32>(EFXAAQuality::MAX) - 1));
 			if (!bApplyDisplayTransform)
 			{
-				FXAAInputs.OverrideOutput = Inputs.OverrideOutput;
+				FXAAInputs.OverrideOutput = Inputs.OverrideOutput.IsValid()
+					? Inputs.OverrideOutput
+					: CreateOutputRenderTarget(GraphBuilder, InView, PassContext.OutputViewRect, Input.Texture->Desc, TEXT("CompositeCorePassFXAA"));
 			}
 
 			FScreenPassTexture Output = AddFXAAPass(GraphBuilder, InView, FXAAInputs);
